Move DiffAssetBulkDataCommandlet to a Program.
#rnx

diff --git a/Engine/Source/Editor/UnrealEd/Private/Commandlets/DiffAssetBulkDataCommandlet.cpp b/Engine/Source/Editor/UnrealEd/Private/Commandlets/DiffAssetBulkDataCommandlet.cpp
--- a/Engine/Source/Editor/UnrealEd/Private/Commandlets/DiffAssetBulkDataCommandlet.cpp
+++ b/Engine/Source/Editor/UnrealEd/Private/Commandlets/DiffAssetBulkDataCommandlet.cpp
@@ -77,6 +77,8 @@ UDiffAssetBulkDataCommandlet::UDiffAssetBulkDataCommandlet(const FObjectInitiali
 
 int32 UDiffAssetBulkDataCommandlet::Main(const FString& FullCommandLine)
 {
+	UE_LOG(LogDiffAssetBulk, Error, TEXT("Has been moved, use the Program DiffAssetBulkData"));
+
 	FString BaseFileName, CurrentFileName;
 	const TCHAR* CmdLine = *FullCommandLine;
 	if (FParse::Value(CmdLine, TEXT("Base="), BaseFileName) == false ||
@@ -166,26 +168,53 @@ int32 UDiffAssetBulkDataCommandlet::Main(const FString& FullCommandLine)
 	};
 	TArray<FIteratedPackage> UnionedPackages;
 
+	uint64 CurrentTotalSize = 0;
+	uint64 BaseTotalSize = 0;
+	uint64 DeletedSize = 0;
+	uint64 NewSize = 0;
+
 	{
 		for (const TPair<FName, const FAssetPackageData*>& NamePackageDataPair : BasePackages)
 		{
 			const FAssetPackageData* Current = CurrentState.GetAssetPackageData(NamePackageDataPair.Key);
 
+			const FAssetData* BaseMIAsset = UE::AssetRegistry::GetMostImportantAsset(BaseState.GetAssetsByPackageName(NamePackageDataPair.Key), UE::AssetRegistry::EGetMostImportantAssetFlags::IgnoreSkipClasses);
+			uint64 BaseCompressedSize = 0;
+			if (BaseMIAsset && BaseMIAsset->GetTagValue(UE::AssetRegistry::Stage_ChunkCompressedSizeFName, BaseCompressedSize))
+			{
+				BaseTotalSize += BaseCompressedSize;
+				if (Current == nullptr)
+				{
+					DeletedSize += BaseCompressedSize;
+				}
+			}
+
 			UnionedPackages.Emplace(FIteratedPackage(NamePackageDataPair.Key, NamePackageDataPair.Value, Current));
 
 			if (Current == nullptr)
 			{
 				DeletedPackages.Add(NamePackageDataPair.Key);
-				continue;
 			}
 		}
 
 		for (const TPair<FName, const FAssetPackageData*>& NamePackageDataPair : CurrentPackages)
 		{
-			if (BaseState.GetAssetPackageData(NamePackageDataPair.Key) == nullptr)
+			const FAssetPackageData* Base = BaseState.GetAssetPackageData(NamePackageDataPair.Key);
+
+			const FAssetData* CurrentMIAsset = UE::AssetRegistry::GetMostImportantAsset(CurrentState.GetAssetsByPackageName(NamePackageDataPair.Key), UE::AssetRegistry::EGetMostImportantAssetFlags::IgnoreSkipClasses);
+			uint64 CurrentCompressedSize = 0;
+			if (CurrentMIAsset && CurrentMIAsset->GetTagValue(UE::AssetRegistry::Stage_ChunkCompressedSizeFName, CurrentCompressedSize))
 			{
-				NewPackages.Add(NamePackageDataPair.Key);
+				CurrentTotalSize += CurrentCompressedSize;
+				if (Base == nullptr)
+				{
+					NewSize += CurrentCompressedSize;
+				}
+			}
 
+			if (Base == nullptr)
+			{
+				NewPackages.Add(NamePackageDataPair.Key);
 				UnionedPackages.Emplace(FIteratedPackage(NamePackageDataPair.Key, nullptr, NamePackageDataPair.Value));
 			}
 		}
@@ -283,6 +312,8 @@ int32 UDiffAssetBulkDataCommandlet::Main(const FString& FullCommandLine)
 	TArray<FName /* PackageName */> PackagesWithUnassignableDiffsAndUntaggedAssets;
 	TMap<FTopLevelAssetPath /* AssetClass */, TArray<FName /* PackageName */>> PackagesWithUnassignableDiffsByAssumedClass;
 	
+	uint64 TotalChangedSize = 0;
+	TMap < FName /* PackageName */, TPair<uint64, uint64>> PackageSizes;
 	for (const FName& ChangedPackageName : ChangedPackages)
 	{
 		TConstArrayView<FAssetData const*> BaseAssetDatas = BaseState.GetAssetsByPackageName(ChangedPackageName);
@@ -301,6 +332,28 @@ int32 UDiffAssetBulkDataCommandlet::Main(const FString& FullCommandLine)
 			const FAssetData* BaseAssetData;
 			const FAssetData* CurrentAssetData;
 		};
+
+		// Get the size change.
+		const FAssetData* BaseMIAsset = UE::AssetRegistry::GetMostImportantAsset(BaseAssetDatas, UE::AssetRegistry::EGetMostImportantAssetFlags::IgnoreSkipClasses);
+		const FAssetData* CurrentMIAsset = UE::AssetRegistry::GetMostImportantAsset(CurrentAssetDatas, UE::AssetRegistry::EGetMostImportantAssetFlags::IgnoreSkipClasses);
+		if (BaseMIAsset && CurrentMIAsset)
+		{
+			// IoStoreUtilities puts the size of the package on the most important asset
+			uint64 BaseCompressedSize = 0;
+			uint64 CurrentCompressedSize = 0;
+			if (BaseMIAsset->GetTagValue(UE::AssetRegistry::Stage_ChunkCompressedSizeFName, BaseCompressedSize) &&
+				CurrentMIAsset->GetTagValue(UE::AssetRegistry::Stage_ChunkCompressedSizeFName, CurrentCompressedSize))
+			{
+				PackageSizes.Add(ChangedPackageName, {BaseCompressedSize, CurrentCompressedSize});
+
+				// All we can really do here is assume the entire package gets resent, which is not likely
+				// in the general case, but it _is_ reasonably likely in the cases where a package's bulk data changes,
+				// which happens to be what we select on.
+				// The counter argument is that it's possible that the bulk data is Very Large (i.e. multiple compression blocks), and only
+				// one block out of the entire thing changed.
+				TotalChangedSize += CurrentCompressedSize;
+			}
+		}
 		
 		// We want to find all the tags that are in both base/current.
 		TMap<FName /* AssetName */, TArray<FDiffTag>> PackageDiffTags;
@@ -429,6 +482,13 @@ int32 UDiffAssetBulkDataCommandlet::Main(const FString& FullCommandLine)
 			}
 		}
 	}
+
+	if (PackageSizes.Num() == 0)
+	{
+		UE_LOG(LogDiffAssetBulk, Display, TEXT("No package sizes found - stage with asset registry writeback (project settings -> packaging) to get package size info"));
+	}
+
+	int32 PackagesWithNoSize = ChangedPackages.Num() - PackageSizes.Num();
 	
 	int32 TotalNewChunks = 0, TotalChangedChunks = 0, TotalDeletedChunks = 0;
 	UE_LOG(LogDiffAssetBulk, Display, TEXT("Modifications By IoStore Chunk (only bulk data tracked atm):"));
@@ -459,11 +519,12 @@ int32 UDiffAssetBulkDataCommandlet::Main(const FString& FullCommandLine)
 
 	UE_LOG(LogDiffAssetBulk, Display, TEXT(""));
 
-	UE_LOG(LogDiffAssetBulk, Display, TEXT("    Base Packages:                %8d"), BasePackages.Num());
-	UE_LOG(LogDiffAssetBulk, Display, TEXT("    Current Packages:             %8d"), CurrentPackages.Num());
-	UE_LOG(LogDiffAssetBulk, Display, TEXT("    Bulk Data Packages Added:     %8d"), NewPackages.Num());
-	UE_LOG(LogDiffAssetBulk, Display, TEXT("    Bulk Data Packages Deleted:   %8d"), DeletedPackages.Num());
-	UE_LOG(LogDiffAssetBulk, Display, TEXT("    Bulk Data Packages Changed:   %8d"), ChangedPackages.Num());
+	UE_LOG(LogDiffAssetBulk, Display, TEXT("    Base Packages:                %8d (%s bytes)"), BasePackages.Num(), *FText::AsNumber(BaseTotalSize).ToString());
+	UE_LOG(LogDiffAssetBulk, Display, TEXT("    Current Packages:             %8d (%s bytes)"), CurrentPackages.Num(), *FText::AsNumber(CurrentTotalSize).ToString());
+	UE_LOG(LogDiffAssetBulk, Display, TEXT("    Bulk Data Packages Added:     %8d (%s bytes"), NewPackages.Num(), *FText::AsNumber(NewSize).ToString());
+	UE_LOG(LogDiffAssetBulk, Display, TEXT("    Bulk Data Packages Deleted:   %8d (%s bytes)"), DeletedPackages.Num(), *FText::AsNumber(DeletedSize).ToString());
+	UE_LOG(LogDiffAssetBulk, Display, TEXT("    Bulk Data Packages Changed:   %8d (%s bytes -- all chunks!)"), ChangedPackages.Num(), *FText::AsNumber(TotalChangedSize).ToString());
+	UE_LOG(LogDiffAssetBulk, Display, TEXT("    Packages with no size info:   %8d"), PackagesWithNoSize);
 
 	if (ChangedPackages.Num() == 0)
 	{
@@ -474,9 +535,20 @@ int32 UDiffAssetBulkDataCommandlet::Main(const FString& FullCommandLine)
 
 	TArray<FName>& CantDetermineAssetClassPackages = NoTagPackagesByAssumedClass.FindOrAdd(FTopLevelAssetPath());
 
+	// Note this output is parsed by build scripts, be sure to fix those up if you change anything here.
 	UE_LOG(LogDiffAssetBulk, Display, TEXT("Changed package breakdown:"));
 	UE_LOG(LogDiffAssetBulk, Display, TEXT("    No blame information available:"));
-	UE_LOG(LogDiffAssetBulk, Display, TEXT("        Can't determine asset class   : %d"), CantDetermineAssetClassPackages.Num());
+	{
+		Algo::Sort(CantDetermineAssetClassPackages, FNameLexicalLess());
+		UE_LOG(LogDiffAssetBulk, Display, TEXT("        Can't determine asset class   : %-7d // Couldn't pick a representative asset in the package. -ListUnrepresented"), CantDetermineAssetClassPackages.Num());
+		if (bListUnrepresented)
+		{
+			for (const FName& PackageName : CantDetermineAssetClassPackages)
+			{
+				UE_LOG(LogDiffAssetBulk, Display, TEXT("            %s"), *PackageName.ToString());
+			}
+		}
+	}
 	for (TPair<FTopLevelAssetPath, TArray<FName>>& ClassPackages : NoTagPackagesByAssumedClass)
 	{
 		if (ClassPackages.Key == FTopLevelAssetPath())
@@ -484,30 +556,29 @@ int32 UDiffAssetBulkDataCommandlet::Main(const FString& FullCommandLine)
 			continue;
 		}
 
-		UE_LOG(LogDiffAssetBulk, Display, TEXT("        %-30s: %d  // -ListNoBlame=%s"), *ClassPackages.Key.ToString(), ClassPackages.Value.Num(), *ClassPackages.Key.ToString());
-		if (ListNoBlame.Compare(TEXT("All"), ESearchCase::IgnoreCase) == 0 ||
-			ListNoBlame.Compare(ClassPackages.Key.ToString(), ESearchCase::IgnoreCase) == 0)
+		uint64 TotalSizes = 0;
+		for (const FName& PackageName : ClassPackages.Value)
 		{
-			for (const FName& PackageName : ClassPackages.Value)
+			TPair<uint64, uint64>* Sizes = PackageSizes.Find(PackageName);
+			if (Sizes)
 			{
-				UE_LOG(LogDiffAssetBulk, Display, TEXT("        %s"), *PackageName.ToString());
+				TotalSizes += Sizes->Value;
 			}
 		}
-	}
 
-	if (CantDetermineAssetClassPackages.Num())
-	{
-		Algo::Sort(CantDetermineAssetClassPackages, FNameLexicalLess());
-		UE_LOG(LogDiffAssetBulk, Display, TEXT("    Can't determine asset class:  : %-7d // Couldn't pick a representative asset in the package. -ListUnrepresented"), CantDetermineAssetClassPackages.Num());
-		if (bListUnrepresented)
+		UE_LOG(LogDiffAssetBulk, Display, TEXT("        %-30s: %d (%s bytes)  // -ListNoBlame=%s"), *ClassPackages.Key.ToString(), ClassPackages.Value.Num(), *FText::AsNumber(TotalSizes).ToString(), *ClassPackages.Key.ToString());
+		if (ListNoBlame.Compare(TEXT("All"), ESearchCase::IgnoreCase) == 0 ||
+			ListNoBlame.Compare(ClassPackages.Key.ToString(), ESearchCase::IgnoreCase) == 0)
 		{
-			for (const FName& PackageName : CantDetermineAssetClassPackages)
+			for (const FName& PackageName : ClassPackages.Value)
 			{
 				UE_LOG(LogDiffAssetBulk, Display, TEXT("        %s"), *PackageName.ToString());
 			}
 		}
 	}
 
+
+
 	if (PackagesWithUnassignableDiffsByAssumedClass.Num())
 	{
 		int32 TotalUnassignablePackages = 0;
@@ -520,7 +591,17 @@ int32 UDiffAssetBulkDataCommandlet::Main(const FString& FullCommandLine)
 		UE_LOG(LogDiffAssetBulk, Display, TEXT("    Can't determine blame:        : %-7d // Assets had blame tags but all matched - check determinism! -ListDeterminism"), TotalUnassignablePackages);
 		for (TPair<FTopLevelAssetPath, TArray<FName>>& ClassPackages : PackagesWithUnassignableDiffsByAssumedClass)
 		{
-			UE_LOG(LogDiffAssetBulk, Display, TEXT("        %s : %d"), *ClassPackages.Key.ToString(), ClassPackages.Value.Num());
+			uint64 TotalSizes = 0;
+			for (const FName& PackageName : ClassPackages.Value)
+			{
+				TPair<uint64, uint64>* Sizes = PackageSizes.Find(PackageName);
+				if (Sizes)
+				{
+					TotalSizes += Sizes->Value;
+				}
+			}
+
+			UE_LOG(LogDiffAssetBulk, Display, TEXT("        %s : %d (%s bytes)"), *ClassPackages.Key.ToString(), ClassPackages.Value.Num(), *FText::AsNumber(TotalSizes).ToString());
 			Algo::Sort(ClassPackages.Value, FNameLexicalLess());
 			if (bListDeterminism)
 			{
diff --git a/Engine/Source/Programs/DiffAssetBulkData/DiffAssetBulkData.Build.cs b/Engine/Source/Programs/DiffAssetBulkData/DiffAssetBulkData.Build.cs
new file mode 100644
--- /dev/null
+++ b/Engine/Source/Programs/DiffAssetBulkData/DiffAssetBulkData.Build.cs
@@ -0,0 +1,20 @@
+// Copyright Epic Games, Inc. All Rights Reserved.
+
+using UnrealBuildTool;
+
+public class DiffAssetBulkData : ModuleRules
+{
+	public DiffAssetBulkData(ReadOnlyTargetRules Target) : base(Target)
+	{
+		PublicIncludePathModuleNames.Add("Launch");
+
+		PrivateDependencyModuleNames.AddRange(new string[] {
+			"ApplicationCore",
+			"AssetRegistry",
+			"Core", 
+			"CoreUObject",
+			"Json",
+			"Projects", 
+		});
+	}
+}
diff --git a/Engine/Source/Programs/DiffAssetBulkData/DiffAssetBulkData.Target.cs b/Engine/Source/Programs/DiffAssetBulkData/DiffAssetBulkData.Target.cs
new file mode 100644
--- /dev/null
+++ b/Engine/Source/Programs/DiffAssetBulkData/DiffAssetBulkData.Target.cs
@@ -0,0 +1,47 @@
+// Copyright Epic Games, Inc. All Rights Reserved.
+
+using UnrealBuildTool;
+using System.Collections.Generic;
+
+public class DiffAssetBulkDataTarget : TargetRules
+{
+	public DiffAssetBulkDataTarget(TargetInfo Target) : base(Target)
+	{
+		Type = TargetType.Program;
+		LinkType = TargetLinkType.Modular;
+		IncludeOrderVersion = EngineIncludeOrderVersion.Latest;
+		LaunchModuleName = "DiffAssetBulkData";
+
+		bUseXGEController					= false;
+		bCompileFreeType					= false;
+		bLoggingToMemoryEnabled				= false;
+		bUseLoggingInShipping				= true;
+		bCompileWithAccessibilitySupport	= false;
+		bCompileWithPluginSupport			= false;
+		bIncludePluginsForTargetPlatforms	= false;
+		bWithServerCode						= false;
+		bCompileNavmeshClusterLinks			= false;
+		bCompileNavmeshSegmentLinks			= false;
+		bCompileRecast						= false;
+		bCompileICU 						= false;
+		bWithLiveCoding						= false;
+		bBuildDeveloperTools				= false;
+		bBuildWithEditorOnlyData			= true;
+		bCompileAgainstEngine				= false;
+		bCompileAgainstCoreUObject			= true;
+		bCompileAgainstApplicationCore		= true;
+		bUsesSlate							= false;
+		bIsBuildingConsoleApplication		= true;
+		bForceBuildTargetPlatforms			= true;
+		bNeedsExtraShaderFormats			= false;
+		bAreTargetSDKVersionsRelevantOverride = false;
+
+		WindowsPlatform.bUseBundledDbgHelp = false;
+		WindowsPlatform.bPixProfilingEnabled = false;
+		WindowsPlatform.bMergeIdenticalCOMDATs = true;
+		WindowsPlatform.bStripUnreferencedSymbols = true;
+
+		bEnableTrace = true;
+		GlobalDefinitions.Add("UE_EXTERNAL_PROFILING_ENABLED=0");
+	}
+}
diff --git a/Engine/Source/Programs/DiffAssetBulkData/Private/DiffAssetBulkData.cpp b/Engine/Source/Programs/DiffAssetBulkData/Private/DiffAssetBulkData.cpp
new file mode 100644
--- /dev/null
+++ b/Engine/Source/Programs/DiffAssetBulkData/Private/DiffAssetBulkData.cpp
@@ -0,0 +1,751 @@
+// Copyright Epic Games, Inc. All Rights Reserved.
+
+#include "RequiredProgramMainCPPInclude.h"
+
+#include "Algo/Sort.h"
+#include "AssetRegistry/AssetData.h"
+#include "AssetRegistry/AssetDataTagMap.h"
+#include "AssetRegistry/AssetRegistryState.h"
+#include "AssetRegistry/IAssetRegistry.h"
+#include "Containers/Array.h"
+#include "Containers/ArrayView.h"
+#include "Containers/ContainersFwd.h"
+#include "Containers/Map.h"
+#include "HAL/PlatformCrt.h"
+#include "IO/IoDispatcher.h"
+#include "IO/IoHash.h"
+#include "Misc/CString.h"
+#include "Misc/Parse.h"
+#include "Templates/UnrealTemplate.h"
+#include "Trace/Detail/Channel.h"
+#include "UObject/NameTypes.h"
+#include "UObject/TopLevelAssetPath.h"
+
+IMPLEMENT_APPLICATION(DiffAssetBulkData, "DiffAssetBulkData");
+
+
+DEFINE_LOG_CATEGORY_STATIC(LogDiffAssetBulk, Display, All);
+
+/**
+ * Diff Asset Bulk Data
+ * 
+ * This loads two asset registries newer than FAssetRegistryVersion::AddedChunkHashes,
+ * and attempts to find the reason for bulk data differences.
+ * 
+ * First, it finds what bulk datas changed by using the hash of the bulk data,
+ * then it uses "Diff Tags" to try and determine at what point during the derived data
+ * build the change occurred.
+ * 
+ * 
+ * Diff Tags
+ * 
+ * Diff Tags are cook tags added during the cook process using Ar.CookContext()->CookTagList() (see CookTagList.h)
+ * and are of the form "Cook_Diff_##_Key":
+ *
+ * 		- "Cook_": 	Added automatically by the the cook tag system.
+ * 		- "Diff_": 	Identifies the tag as a diff tag.
+ * 		- "##":		Specifies where in the build process the tag represents (Ordering).
+ * 		- "_Key":	Descriptive text for the tag.
+ * 
+ * If a bulk data difference is found, the diff tags are checked for differences in order, and the first
+ * diff tag that changed is assigned the "blame" for the change under the assumption that later
+ * tags will necessarily change as a result of the earlier change.
+ * 
+ * If diff tags are present for the asset and none of the diff tags changed, then it is assumed that a build determinism 
+ * issue has caused the change.
+ *
+ */
+
+/**
+ *  The list of known cook diff tags - this is just used to provide explanations in the output for the reader.
+ */
+static struct FBuiltinDiffTagHelp {const TCHAR* TagName; const TCHAR* TagHelp;} GBuiltinDiffTagHelp[] = 
+{
+	{TEXT("Cook_Diff_20_Tex2D_CacheKey"), TEXT("Texture settings or referenced data changed (DDC2)")},
+	{TEXT("Cook_Diff_20_Tex2D_DDK"), TEXT("Texture settings or referenced data changed (DDC1)")},
+	{TEXT("Cook_Diff_10_Tex2D_Source"), TEXT("Texture source data changed")}
+};
+
+
+static int32 RunDiffAssetBulkData()
+{
+	FString BaseFileName, CurrentFileName;
+	const TCHAR* CmdLine = FCommandLine::Get();
+	if (FParse::Value(CmdLine, TEXT("Base="), BaseFileName) == false ||
+		FParse::Value(CmdLine, TEXT("Current="), CurrentFileName) == false)
+	{
+		UE_LOG(LogDiffAssetBulk, Display, TEXT(""));
+		UE_LOG(LogDiffAssetBulk, Display, TEXT("Diff Asset Bulk Data"));
+		UE_LOG(LogDiffAssetBulk, Display, TEXT(""));
+		UE_LOG(LogDiffAssetBulk, Display, TEXT("Loads two development asset registries and finds all bulk data changes, and tries to find why"));
+		UE_LOG(LogDiffAssetBulk, Display, TEXT("the bulk data changed. Development asset registries are in the cooked /Metadata directory."));
+		UE_LOG(LogDiffAssetBulk, Display, TEXT(""));
+		UE_LOG(LogDiffAssetBulk, Display, TEXT("Parameters:"));
+		UE_LOG(LogDiffAssetBulk, Display, TEXT(""));
+		UE_LOG(LogDiffAssetBulk, Display, TEXT("    -Base=<path/to/file>              Base Development Asset Registry (Required)"));
+		UE_LOG(LogDiffAssetBulk, Display, TEXT("    -Current=<path/to/file>           New Development Asset Registry (Required)"));
+		UE_LOG(LogDiffAssetBulk, Display, TEXT("    -ListMixed                        Show the list of changed packages with assets that have matching"));
+		UE_LOG(LogDiffAssetBulk, Display, TEXT("                                      blame tags, but also assets without."));
+		UE_LOG(LogDiffAssetBulk, Display, TEXT("    -ListDeterminism                  Show the list of changed packages with assets that have matching"));
+		UE_LOG(LogDiffAssetBulk, Display, TEXT("                                      blame tags."));		
+		UE_LOG(LogDiffAssetBulk, Display, TEXT("    -ListBlame=<blame tag>            Show the list of assets that changed due to a specific blame"));
+		UE_LOG(LogDiffAssetBulk, Display, TEXT("                                      tag or \"All\" to list all changed assets with known blame."));
+		UE_LOG(LogDiffAssetBulk, Display, TEXT("    -ListUnrepresented                Show the list of packages where a representative asset couldn't be found.")); 
+		UE_LOG(LogDiffAssetBulk, Display, TEXT("    -ListNoBlame=<class>              Show the list of assets that changed for a specific class, or \"All\""));
+		UE_LOG(LogDiffAssetBulk, Display, TEXT("    -ListCSV=<filename>               Write all changed packages to the given CSV file."));
+		return 1;
+	}
+
+	bool bListMixed = FParse::Param(CmdLine, TEXT("ListMixed"));
+	bool bListDeterminism = FParse::Param(CmdLine, TEXT("ListDeterminism"));
+	bool bListUnrepresented = FParse::Param(CmdLine, TEXT("ListUnrepresented"));
+	FString ListBlame;
+	FParse::Value(CmdLine, TEXT("ListBlame="), ListBlame);
+	FString ListNoBlame;
+	FParse::Value(CmdLine, TEXT("ListNoBlame="), ListNoBlame);
+
+	FString ListCSV;
+	TUniquePtr<FArchive> OutputCSVAr;
+	if (FParse::Value(CmdLine, TEXT("ListCSV="), ListCSV))
+	{
+		OutputCSVAr.Reset(IFileManager::Get().CreateFileWriter(*ListCSV, 0));
+		if (!OutputCSVAr)
+		{
+			UE_LOG(LogDiffAssetBulk, Error, TEXT("Unable to open output CSV file: %s"), *ListCSV);
+			return false;
+		}
+		OutputCSVAr->Logf(TEXT("Blame, Class, PackageName, BlameBefore, BlameAfter"));
+	}
+
+	// Convert the static init help text to a map
+	TMap<FName, const TCHAR*> BuiltinDiffTagHelpMap;
+	for (FBuiltinDiffTagHelp& DiffTagHelp : GBuiltinDiffTagHelp)
+	{
+		BuiltinDiffTagHelpMap.Add(DiffTagHelp.TagName, DiffTagHelp.TagHelp);
+	}
+
+	FAssetRegistryState BaseState, CurrentState;
+	FAssetRegistryVersion::Type BaseVersion, CurrentVersion;
+	UE_LOG(LogDiffAssetBulk, Display, TEXT("Loading Base... (%s)"), *BaseFileName);
+	if (FAssetRegistryState::LoadFromDisk(*BaseFileName, FAssetRegistryLoadOptions(), BaseState, &BaseVersion) == false)
+	{
+		UE_LOG(LogDiffAssetBulk, Error, TEXT("Failed load base (%s)"), *BaseFileName);
+		return 1;
+	}
+	UE_LOG(LogDiffAssetBulk, Display, TEXT("Loading Current... (%s)"), *CurrentFileName);
+	if (FAssetRegistryState::LoadFromDisk(*CurrentFileName, FAssetRegistryLoadOptions(), CurrentState, &CurrentVersion) == false)
+	{
+		UE_LOG(LogDiffAssetBulk, Error, TEXT("Failed load current (%s)"), *CurrentFileName);
+		return 1;
+	}
+
+	//
+	// The cook process adds the hash for almost all iochunks to the asset registry - 
+	// so as long as both asset registries have that data, we get what we want.
+	//
+	if (BaseVersion < FAssetRegistryVersion::AddedChunkHashes)
+	{
+		UE_LOG(LogDiffAssetBulk, Error, TEXT("Base asset registry version is too old (%d, need %d)"), BaseVersion, FAssetRegistryVersion::AddedChunkHashes);
+		return 1;
+	}
+	if (CurrentVersion < FAssetRegistryVersion::AddedChunkHashes)
+	{
+		UE_LOG(LogDiffAssetBulk, Error, TEXT("Current asset registry version is too old (%d, need %d)"), CurrentVersion, FAssetRegistryVersion::AddedChunkHashes);
+		return 1;
+	}
+
+	const TMap<FName, const FAssetPackageData*>& BasePackages = BaseState.GetAssetPackageDataMap();
+	const TMap<FName, const FAssetPackageData*>& CurrentPackages = CurrentState.GetAssetPackageDataMap();
+
+	TArray<FName> NewPackages, DeletedPackages;
+
+	struct FIteratedPackage
+	{
+		FName Name = NAME_None;
+		const FAssetPackageData* Base = nullptr;
+		const FAssetPackageData* Current = nullptr;
+		FIteratedPackage() = default;
+		FIteratedPackage(FName _Name, const FAssetPackageData* _Base, const FAssetPackageData* _Current) :
+			Name(_Name),
+			Base(_Base),
+			Current(_Current) {}
+
+	};
+	TArray<FIteratedPackage> UnionedPackages;
+
+	uint64 CurrentTotalSize = 0;
+	uint64 BaseTotalSize = 0;
+	uint64 DeletedSize = 0;
+	uint64 NewSize = 0;
+
+	{
+		for (const TPair<FName, const FAssetPackageData*>& NamePackageDataPair : BasePackages)
+		{
+			const FAssetPackageData* Current = CurrentState.GetAssetPackageData(NamePackageDataPair.Key);
+
+			const FAssetData* BaseMIAsset = UE::AssetRegistry::GetMostImportantAsset(BaseState.GetAssetsByPackageName(NamePackageDataPair.Key), UE::AssetRegistry::EGetMostImportantAssetFlags::IgnoreSkipClasses);
+			uint64 BaseCompressedSize = 0;
+			if (BaseMIAsset && BaseMIAsset->GetTagValue(UE::AssetRegistry::Stage_ChunkCompressedSizeFName, BaseCompressedSize))
+			{
+				BaseTotalSize += BaseCompressedSize;
+				if (Current == nullptr)
+				{
+					DeletedSize += BaseCompressedSize;
+				}
+			}
+
+			UnionedPackages.Emplace(FIteratedPackage(NamePackageDataPair.Key, NamePackageDataPair.Value, Current));
+
+			if (Current == nullptr)
+			{
+				DeletedPackages.Add(NamePackageDataPair.Key);
+			}
+		}
+
+		for (const TPair<FName, const FAssetPackageData*>& NamePackageDataPair : CurrentPackages)
+		{
+			const FAssetPackageData* Base = BaseState.GetAssetPackageData(NamePackageDataPair.Key);
+
+			const FAssetData* CurrentMIAsset = UE::AssetRegistry::GetMostImportantAsset(CurrentState.GetAssetsByPackageName(NamePackageDataPair.Key), UE::AssetRegistry::EGetMostImportantAssetFlags::IgnoreSkipClasses);
+			uint64 CurrentCompressedSize = 0;
+			if (CurrentMIAsset && CurrentMIAsset->GetTagValue(UE::AssetRegistry::Stage_ChunkCompressedSizeFName, CurrentCompressedSize))
+			{
+				CurrentTotalSize += CurrentCompressedSize;
+				if (Base == nullptr)
+				{
+					NewSize += CurrentCompressedSize;
+				}
+			}
+
+			if (Base == nullptr)
+			{
+				NewPackages.Add(NamePackageDataPair.Key);
+				UnionedPackages.Emplace(FIteratedPackage(NamePackageDataPair.Key, nullptr, NamePackageDataPair.Value));
+			}
+		}
+	}
+
+	// Now we need to see what changed.
+	//
+	// This whole thing assumes that the index parameter of CreateIoChunkId is always 0. This is likely not going
+	// to be true with FDerivedData, once that gets turned on, but should be easy to update when the time comes.
+	//
+	TSet<FName> ChangedChunksByType[(uint32)EIoChunkType::MAX], NewChunksByType[(uint32)EIoChunkType::MAX], DeletedChunksByType[(uint32)EIoChunkType::MAX];
+	for (const FIteratedPackage& IteratedPackage : UnionedPackages)
+	{
+		const FAssetPackageData* BasePackage = IteratedPackage.Base;
+		const FAssetPackageData* CurrentPackage = IteratedPackage.Current;
+
+		if (BasePackage)
+		{
+			for (const TPair<FIoChunkId, FIoHash>& ChunkHashPair : BasePackage->ChunkHashes)
+			{
+				if (ChunkHashPair.Key.GetChunkType() != EIoChunkType::BulkData &&
+					ChunkHashPair.Key.GetChunkType() != EIoChunkType::OptionalBulkData &&
+					ChunkHashPair.Key.GetChunkType() != EIoChunkType::MemoryMappedBulkData)
+					continue;
+
+				const FIoHash* CurrentHash = nullptr;
+				if (CurrentPackage)
+				{
+					CurrentHash = CurrentPackage->ChunkHashes.Find(ChunkHashPair.Key);
+				}
+
+				if (CurrentHash == nullptr)
+				{
+					TSet<FName>& Deleted = DeletedChunksByType[(uint32)ChunkHashPair.Key.GetChunkType()];
+					check(Deleted.Contains(IteratedPackage.Name) == false); // Because only 0 chunk index
+					Deleted.Add(IteratedPackage.Name);
+					continue;
+				}
+
+				if (*CurrentHash != ChunkHashPair.Value)
+				{
+					TSet<FName>& Changed = ChangedChunksByType[(uint32)ChunkHashPair.Key.GetChunkType()];
+					check(Changed.Contains(IteratedPackage.Name) == false); // Because only 0 chunk index
+					Changed.Add(IteratedPackage.Name);
+				}
+			}
+		}
+
+		if (CurrentPackage)
+		{
+			for (const TPair<FIoChunkId, FIoHash>& ChunkHashPair : CurrentPackage->ChunkHashes)
+			{
+				if (ChunkHashPair.Key.GetChunkType() != EIoChunkType::BulkData &&
+					ChunkHashPair.Key.GetChunkType() != EIoChunkType::OptionalBulkData &&
+					ChunkHashPair.Key.GetChunkType() != EIoChunkType::MemoryMappedBulkData)
+					continue;
+
+				if (!BasePackage ||
+					BasePackage->ChunkHashes.Contains(ChunkHashPair.Key) == false)
+				{
+					TSet<FName>& New = NewChunksByType[(uint32)ChunkHashPair.Key.GetChunkType()];
+					check(New.Contains(IteratedPackage.Name) == false); // Because only 0 chunk index
+					New.Add(IteratedPackage.Name);
+				}
+			}
+		}
+	}
+
+	// Get a unique list of changed packages.
+	TSet<FName> ChangedPackages;
+	for (uint32 ChunkTypeIndex = 0; ChunkTypeIndex < (uint32)EIoChunkType::MAX; ChunkTypeIndex++)
+	{
+		ChangedPackages.Append(ChangedChunksByType[ChunkTypeIndex]);
+	}
+
+
+	//
+	// We know what bulk datas *packages* changed. Try and see if any of the assets in the package have
+	// diff blame tags for us to determine cause. _usually_ there's one asset per package, but it's definitely
+	// possible to have more. Additionally _usually_ there's a good single candidate for assigning the data
+	// cost, however it is possible to have e.g. an importer create a lot of assets in a single package that
+	// all add bulk data to the package.
+	// 
+	// Once we have FDerivedData we might be able to keep what data belongs to which asset.
+	//
+	struct FDiffResult
+	{
+		FString ChangedAssetObjectPath;
+		FString TagBaseValue;
+		FString TagCurrentValue;
+	};
+
+	TMap<FName /* TagName */, TMap<FTopLevelAssetPath /* AssetClass */, TArray<FDiffResult>>> Results;
+	TMap<FTopLevelAssetPath /* AssetClass */, TArray<FName /* PackageName */>> NoTagPackagesByAssumedClass;
+	TArray<FName /* PackageName */> PackagesWithUnassignableDiffsAndUntaggedAssets;
+	TMap<FTopLevelAssetPath /* AssetClass */, TArray<FName /* PackageName */>> PackagesWithUnassignableDiffsByAssumedClass;
+	
+	uint64 TotalChangedSize = 0;
+	TMap < FName /* PackageName */, TPair<uint64, uint64>> PackageSizes;
+	for (const FName& ChangedPackageName : ChangedPackages)
+	{
+		TConstArrayView<FAssetData const*> BaseAssetDatas = BaseState.GetAssetsByPackageName(ChangedPackageName);
+		TConstArrayView<FAssetData const*> CurrentAssetDatas = CurrentState.GetAssetsByPackageName(ChangedPackageName);
+
+		struct FDiffTag
+		{
+			// Order is used to sort the diff blame keys so that the correct thing is blamed. This is
+			// so that e.g. changing the texture source (which would change the ddc key) gets properly blamed
+			// as it is lower order.
+			int Order;
+			FName TagName;
+			FString BaseValue;
+			FString CurrentValue;
+
+			const FAssetData* BaseAssetData;
+			const FAssetData* CurrentAssetData;
+		};
+
+		// Get the size change.
+		const FAssetData* BaseMIAsset = UE::AssetRegistry::GetMostImportantAsset(BaseAssetDatas, UE::AssetRegistry::EGetMostImportantAssetFlags::IgnoreSkipClasses);
+		const FAssetData* CurrentMIAsset = UE::AssetRegistry::GetMostImportantAsset(CurrentAssetDatas, UE::AssetRegistry::EGetMostImportantAssetFlags::IgnoreSkipClasses);
+		if (BaseMIAsset && CurrentMIAsset)
+		{
+			// IoStoreUtilities puts the size of the package on the most important asset
+			uint64 BaseCompressedSize = 0;
+			uint64 CurrentCompressedSize = 0;
+			if (BaseMIAsset->GetTagValue(UE::AssetRegistry::Stage_ChunkCompressedSizeFName, BaseCompressedSize) &&
+				CurrentMIAsset->GetTagValue(UE::AssetRegistry::Stage_ChunkCompressedSizeFName, CurrentCompressedSize))
+			{
+				PackageSizes.Add(ChangedPackageName, {BaseCompressedSize, CurrentCompressedSize});
+
+				// All we can really do here is assume the entire package gets resent, which is not likely
+				// in the general case, but it _is_ reasonably likely in the cases where a package's bulk data changes,
+				// which happens to be what we select on.
+				// The counter argument is that it's possible that the bulk data is Very Large (i.e. multiple compression blocks), and only
+				// one block out of the entire thing changed.
+				TotalChangedSize += CurrentCompressedSize;
+			}
+		}
+		
+		// We want to find all the tags that are in both base/current.
+		TMap<FName /* AssetName */, TArray<FDiffTag>> PackageDiffTags;
+		bool bPackageHasUntaggedAsset = false;
+		for (const FAssetData* BaseAssetData : BaseAssetDatas)
+		{
+			BaseAssetData->EnumerateTags([&PackageDiffTags, BaseAssetData, CurrentAssetDatas](TPair<FName, FAssetTagValueRef> TagAndValue)
+			{
+				TCHAR Name[NAME_SIZE];
+				TagAndValue.Key.GetPlainNameString(Name);
+				if (FCString::Strncmp(Name, TEXT("Cook_Diff_"), 10))
+				{
+					return;
+				}
+
+				// This is O(N) but like 99.9% of the time there's only 1 asset.
+				const FAssetData* const* CurrentAssetData = CurrentAssetDatas.FindByPredicate([SearchAssetName = &BaseAssetData->AssetName](const FAssetData* AssetData) { return (AssetData->AssetName == *SearchAssetName); });
+				if (CurrentAssetData == nullptr)
+				{
+					return;
+				}
+
+				FString CurrentValue;
+				if (CurrentAssetData[0]->GetTagValue(TagAndValue.Key, CurrentValue) == false)
+				{
+					// Both version don't have the tag so we can't compare.
+					return;
+				}
+
+				TArray<FDiffTag>& AssetDiffTags = PackageDiffTags.FindOrAdd(BaseAssetData->AssetName);
+				FDiffTag& Tag = AssetDiffTags.AddDefaulted_GetRef();
+				Tag.Order = FCString::Atoi(Name + FCString::Strlen(TEXT("Cook_Diff_"))); // this gets optimized to +10
+				Tag.TagName = TagAndValue.Key;
+				Tag.BaseValue =  TagAndValue.Value.AsString();
+				Tag.CurrentValue = MoveTemp(CurrentValue);
+				Tag.BaseAssetData = BaseAssetData;
+				Tag.CurrentAssetData = *CurrentAssetData;
+			});
+
+			if (PackageDiffTags.Contains(BaseAssetData->AssetName) == false)
+			{
+				bPackageHasUntaggedAsset = true;
+				// An asset exists in the package that doesn't have any tags - make a note so that
+				// we can suggest this caused the bulk data diff if we don't find a blame.
+			}
+		}
+
+		bool bPackageHasUntaggedAndTaggedAssets = false;
+		if (PackageDiffTags.Num())
+		{
+			if (bPackageHasUntaggedAsset)
+			{
+				bPackageHasUntaggedAndTaggedAssets = true;
+			}
+		}
+		else
+		{
+			// Nothing has anything to use for diff blaming for this package.
+			// Try to find a representative asset class from the assets in the package.
+			FAssetData const* RepresentativeAsset = UE::AssetRegistry::GetMostImportantAsset(CurrentAssetDatas, UE::AssetRegistry::EGetMostImportantAssetFlags::RequireOneTopLevelAsset);
+			if (RepresentativeAsset == nullptr)
+			{
+				NoTagPackagesByAssumedClass.FindOrAdd(FTopLevelAssetPath()).Add(ChangedPackageName);
+			}
+			else
+			{
+				NoTagPackagesByAssumedClass.FindOrAdd(RepresentativeAsset->AssetClassPath).Add(ChangedPackageName);
+			}
+			
+			continue;
+		}
+
+		// Now we check and see if any of the diff tags can tell us why the package changed.
+		// We could find multiple assets that caused the change.
+		bool bFoundDiffTag = false;
+		for (TPair<FName, TArray<FDiffTag>>& AssetDiffTagPair : PackageDiffTags)
+		{
+			TArray<FDiffTag>& AssetDiffTags = AssetDiffTagPair.Value;
+			Algo::SortBy(AssetDiffTags, &FDiffTag::Order);
+			
+			for (FDiffTag& Tag : AssetDiffTags)
+			{
+				if (Tag.BaseValue != Tag.CurrentValue)
+				{
+					TMap<FTopLevelAssetPath, TArray<FDiffResult>>& TagResults = Results.FindOrAdd(Tag.TagName);
+
+					TArray<FDiffResult>& ClassResults = TagResults.FindOrAdd(Tag.BaseAssetData->AssetClassPath);
+
+					FDiffResult& Result = ClassResults.AddDefaulted_GetRef();
+					Result.ChangedAssetObjectPath = Tag.BaseAssetData->GetObjectPathString();
+					Result.TagBaseValue = MoveTemp(Tag.BaseValue);
+					Result.TagCurrentValue = MoveTemp(Tag.CurrentValue);
+					bFoundDiffTag = true;
+					break;
+				}
+			}
+		}
+
+
+		if (bFoundDiffTag == false)
+		{
+			// This means that all the tags they added didn't change, but the asset did.
+			// Assuming that a DDC key tag has been added, this means either:
+			// 
+			// A) The asset changed independent of DDC key, which is a build consistency / determinism alert.
+			// B) The package had an asset with tags and an asset without tags, and the asset without tags caused
+			//	  the bulk data change.
+			//
+			// Unfortunately A) is a Big Deal and needs a warning, but B might end up being common due to blueprint classes,
+			// so we segregate the lists.
+			if (bPackageHasUntaggedAndTaggedAssets)
+			{
+				PackagesWithUnassignableDiffsAndUntaggedAssets.Add(ChangedPackageName);
+			}
+			else
+			{
+				FAssetData const* RepresentativeAsset = UE::AssetRegistry::GetMostImportantAsset(CurrentAssetDatas, UE::AssetRegistry::EGetMostImportantAssetFlags::RequireOneTopLevelAsset);
+				if (RepresentativeAsset == nullptr)
+				{
+					PackagesWithUnassignableDiffsByAssumedClass.FindOrAdd(FTopLevelAssetPath()).Add(ChangedPackageName);
+				}
+				else
+				{
+					PackagesWithUnassignableDiffsByAssumedClass.FindOrAdd(RepresentativeAsset->AssetClassPath).Add(ChangedPackageName);
+				}
+			}
+		}
+	}
+
+	if (PackageSizes.Num() == 0)
+	{
+		UE_LOG(LogDiffAssetBulk, Display, TEXT("No package sizes found - stage with asset registry writeback (project settings -> packaging) to get package size info"));
+	}
+
+	int32 PackagesWithNoSize = ChangedPackages.Num() - PackageSizes.Num();
+	
+	int32 TotalNewChunks = 0, TotalChangedChunks = 0, TotalDeletedChunks = 0;
+	UE_LOG(LogDiffAssetBulk, Display, TEXT("Modifications By IoStore Chunk (only bulk data tracked atm):"));
+	UE_LOG(LogDiffAssetBulk, Display, TEXT(""));
+	UE_LOG(LogDiffAssetBulk, Display, TEXT("    ChunkType                   New    Deleted    Changed"));
+	for (uint32 ChunkTypeIndex = 0; ChunkTypeIndex < (uint8)EIoChunkType::MAX; ChunkTypeIndex++)
+	{
+		EIoChunkType ChunkType = (EIoChunkType)ChunkTypeIndex;
+		if (ChunkType != EIoChunkType::BulkData &&
+			ChunkType != EIoChunkType::OptionalBulkData &&
+			ChunkType != EIoChunkType::MemoryMappedBulkData)
+		{
+			continue;
+		}
+
+		const TSet<FName>& NewChunksForType = NewChunksByType[ChunkTypeIndex];
+		const TSet<FName>& DeletedChunksForType = DeletedChunksByType[ChunkTypeIndex];
+		const TSet<FName>& ChangedChunksForType = ChangedChunksByType[ChunkTypeIndex];
+
+		TotalNewChunks += NewChunksForType.Num();
+		TotalChangedChunks += ChangedChunksForType.Num();
+		TotalDeletedChunks += DeletedChunksForType.Num();
+		UE_LOG(LogDiffAssetBulk, Display, TEXT("    %-20s %10d %10d %10d"), *LexToString(ChunkType), NewChunksForType.Num(), DeletedChunksForType.Num(), ChangedChunksForType.Num());
+	}
+	UE_LOG(LogDiffAssetBulk, Display, TEXT("    ====================================================="));
+
+	UE_LOG(LogDiffAssetBulk, Display, TEXT("    %-20s %10d %10d %10d"), TEXT("Total"), TotalNewChunks, TotalDeletedChunks, TotalChangedChunks);
+
+	UE_LOG(LogDiffAssetBulk, Display, TEXT(""));
+
+	UE_LOG(LogDiffAssetBulk, Display, TEXT("    Base Packages:                %8d (%s bytes)"), BasePackages.Num(), *FText::AsNumber(BaseTotalSize).ToString());
+	UE_LOG(LogDiffAssetBulk, Display, TEXT("    Current Packages:             %8d (%s bytes)"), CurrentPackages.Num(), *FText::AsNumber(CurrentTotalSize).ToString());
+	UE_LOG(LogDiffAssetBulk, Display, TEXT("    Bulk Data Packages Added:     %8d (%s bytes)"), NewPackages.Num(), *FText::AsNumber(NewSize).ToString());
+	UE_LOG(LogDiffAssetBulk, Display, TEXT("    Bulk Data Packages Deleted:   %8d (%s bytes)"), DeletedPackages.Num(), *FText::AsNumber(DeletedSize).ToString());
+	UE_LOG(LogDiffAssetBulk, Display, TEXT("    Bulk Data Packages Changed:   %8d (%s bytes -- all chunks!)"), ChangedPackages.Num(), *FText::AsNumber(TotalChangedSize).ToString());
+	UE_LOG(LogDiffAssetBulk, Display, TEXT("    Packages with no size info:   %8d"), PackagesWithNoSize);
+
+	if (ChangedPackages.Num() == 0)
+	{
+		return 0;
+	}
+
+	UE_LOG(LogDiffAssetBulk, Display, TEXT(""));
+
+	TArray<FName>& CantDetermineAssetClassPackages = NoTagPackagesByAssumedClass.FindOrAdd(FTopLevelAssetPath());
+
+	// Note this output is parsed by build scripts, be sure to fix those up if you change anything here.
+	UE_LOG(LogDiffAssetBulk, Display, TEXT("Changed package breakdown:"));
+	UE_LOG(LogDiffAssetBulk, Display, TEXT("    No blame information available:"));
+	{
+		Algo::Sort(CantDetermineAssetClassPackages, FNameLexicalLess());
+		UE_LOG(LogDiffAssetBulk, Display, TEXT("        Can't determine asset class   : %-7d // Couldn't pick a representative asset in the package. -ListUnrepresented"), CantDetermineAssetClassPackages.Num());
+		if (bListUnrepresented)
+		{
+			for (const FName& PackageName : CantDetermineAssetClassPackages)
+			{
+				UE_LOG(LogDiffAssetBulk, Display, TEXT("            %s"), *PackageName.ToString());
+			}
+		}
+		if (OutputCSVAr.IsValid())
+		{
+			for (const FName& PackageName : CantDetermineAssetClassPackages)
+			{
+				OutputCSVAr->Logf(TEXT("NoBlameInfo, Unknown, %s,,"), *WriteToString<64>(PackageName));
+			}
+		}
+	}
+	for (TPair<FTopLevelAssetPath, TArray<FName>>& ClassPackages : NoTagPackagesByAssumedClass)
+	{
+		if (ClassPackages.Key == FTopLevelAssetPath())
+		{
+			continue;
+		}
+
+		uint64 TotalSizes = 0;
+		for (const FName& PackageName : ClassPackages.Value)
+		{
+			TPair<uint64, uint64>* Sizes = PackageSizes.Find(PackageName);
+			if (Sizes)
+			{
+				TotalSizes += Sizes->Value;
+			}
+		}
+
+		TStringBuilder<64> ClassName;
+		ClassName << ClassPackages.Key;
+
+		UE_LOG(LogDiffAssetBulk, Display, TEXT("        %-30s: %d (%s bytes)  // -ListNoBlame=%s"), *ClassName, ClassPackages.Value.Num(), *FText::AsNumber(TotalSizes).ToString(), *ClassName);
+		if (ListNoBlame.Compare(TEXT("All"), ESearchCase::IgnoreCase) == 0 ||
+			ListNoBlame.Compare(ClassPackages.Key.ToString(), ESearchCase::IgnoreCase) == 0)
+		{
+			for (const FName& PackageName : ClassPackages.Value)
+			{
+				UE_LOG(LogDiffAssetBulk, Display, TEXT("        %s"), *WriteToString<64>(PackageName));
+			}
+		}
+		if (OutputCSVAr.IsValid())
+		{
+			for (const FName& PackageName : ClassPackages.Value)
+			{
+				OutputCSVAr->Logf(TEXT("NoBlameInfo, %s, %s,,"), *ClassName, *WriteToString<64>(PackageName));
+			}
+		}
+	}
+
+
+
+	if (PackagesWithUnassignableDiffsByAssumedClass.Num())
+	{
+		int32 TotalUnassignablePackages = 0;
+		for (TPair<FTopLevelAssetPath, TArray<FName>>& ClassPackages : PackagesWithUnassignableDiffsByAssumedClass)
+		{
+			TotalUnassignablePackages += ClassPackages.Value.Num();
+		}
+		
+
+		UE_LOG(LogDiffAssetBulk, Display, TEXT("    Can't determine blame:        : %-7d // Assets had blame tags but all matched - check determinism! -ListDeterminism"), TotalUnassignablePackages);
+		for (TPair<FTopLevelAssetPath, TArray<FName>>& ClassPackages : PackagesWithUnassignableDiffsByAssumedClass)
+		{
+			uint64 TotalSizes = 0;
+			for (const FName& PackageName : ClassPackages.Value)
+			{
+				TPair<uint64, uint64>* Sizes = PackageSizes.Find(PackageName);
+				if (Sizes)
+				{
+					TotalSizes += Sizes->Value;
+				}
+			}
+
+			TStringBuilder<64> ClassName;
+			ClassName << ClassPackages.Key;
+
+			UE_LOG(LogDiffAssetBulk, Display, TEXT("        %s : %d (%s bytes)"), *ClassName, ClassPackages.Value.Num(), *FText::AsNumber(TotalSizes).ToString());
+			Algo::Sort(ClassPackages.Value, FNameLexicalLess());
+			if (bListDeterminism)
+			{
+				for (const FName& PackageName : ClassPackages.Value)
+				{
+					UE_LOG(LogDiffAssetBulk, Display, TEXT("            %s"), *WriteToString<64>(PackageName));
+				}
+			}
+			if (OutputCSVAr.IsValid())
+			{
+				for (const FName& PackageName : ClassPackages.Value)
+				{
+					OutputCSVAr->Logf(TEXT("NonDetermistic, %s, %s,,"), *ClassName, *WriteToString<64>(PackageName));
+				}
+			}
+		}
+	}
+
+	if (PackagesWithUnassignableDiffsAndUntaggedAssets.Num())
+	{
+		Algo::Sort(PackagesWithUnassignableDiffsAndUntaggedAssets, FNameLexicalLess());
+		
+		UE_LOG(LogDiffAssetBulk, Display, TEXT("    Potential untagged assets:    : %-7d // Package had assets with blame tags that matched, but also untagged assets. Might be determinism! -ListMixed"), PackagesWithUnassignableDiffsAndUntaggedAssets.Num());
+		if (bListMixed)
+		{
+			for (const FName& PackageName : PackagesWithUnassignableDiffsAndUntaggedAssets)
+			{
+				UE_LOG(LogDiffAssetBulk, Display, TEXT("        %s"), *PackageName.ToString());
+			}
+		}
+		if (OutputCSVAr.IsValid())
+		{
+			for (const FName& PackageName : PackagesWithUnassignableDiffsAndUntaggedAssets)
+			{
+				OutputCSVAr->Logf(TEXT("Mixed, Unknown, %s,,"), *WriteToString<64>(PackageName));
+			}
+		}
+	}
+
+	if (Results.Num())
+	{
+		UE_LOG(LogDiffAssetBulk, Display, TEXT("    Summary changes by blame tag:"));
+
+		for (TPair<FName, TMap<FTopLevelAssetPath, TArray<FDiffResult>>>& TagResults : Results)
+		{
+			uint32 TagCount = 0;
+			for (TPair<FTopLevelAssetPath, TArray<FDiffResult>>& ClassResults : TagResults.Value)
+			{
+				TagCount += ClassResults.Value.Num();
+			}
+
+			const TCHAR** TagHelp = BuiltinDiffTagHelpMap.Find(TagResults.Key);
+			if (TagHelp != nullptr)
+			{
+				UE_LOG(LogDiffAssetBulk, Display, TEXT("        %-30s: %-7d // %s"), *TagResults.Key.ToString(), TagCount, *TagHelp);
+			}
+			else
+			{
+				UE_LOG(LogDiffAssetBulk, Display, TEXT("        %-30s: %-7d"), *TagResults.Key.ToString(), TagCount);
+			}
+		}
+
+		UE_LOG(LogDiffAssetBulk, Display, TEXT("    Asset changes by blame tag:"));
+
+		for (TPair<FName, TMap<FTopLevelAssetPath, TArray<FDiffResult>>>& TagResults : Results)
+		{
+			UE_LOG(LogDiffAssetBulk, Display, TEXT("        %s  // -ListBlame=%s"), *TagResults.Key.ToString(), *TagResults.Key.ToString());
+
+			for (TPair<FTopLevelAssetPath, TArray<FDiffResult>>& ClassResults : TagResults.Value)
+			{
+				Algo::SortBy(ClassResults.Value, &FDiffResult::ChangedAssetObjectPath);
+				UE_LOG(LogDiffAssetBulk, Display, TEXT("            %s [%d]"), *ClassResults.Key.ToString(), ClassResults.Value.Num());
+
+				if (ListBlame.Compare(TEXT("All"), ESearchCase::IgnoreCase) == 0 ||
+					ListBlame.Compare(ClassResults.Key.ToString(), ESearchCase::IgnoreCase) == 0)
+				{
+					for (FDiffResult& Result : ClassResults.Value)
+					{
+						UE_LOG(LogDiffAssetBulk, Display, TEXT("                %s [%s -> %s]"), *Result.ChangedAssetObjectPath, *Result.TagBaseValue, *Result.TagCurrentValue);
+					}
+				}
+				if (OutputCSVAr.IsValid())
+				{
+					for (FDiffResult& Result : ClassResults.Value)
+					{
+						OutputCSVAr->Logf(TEXT("%s, %s, %s, %s, %s"), *TagResults.Key.ToString(), *ClassResults.Key.ToString(), *WriteToString<64>(Result.ChangedAssetObjectPath), *Result.TagBaseValue, *Result.TagCurrentValue);
+					}
+				}
+			}
+		}
+	}
+
+	UE_LOG(LogDiffAssetBulk, Display, TEXT("Done."));
+
+	return 0;
+}
+
+
+INT32_MAIN_INT32_ARGC_TCHAR_ARGV()
+{
+	FTaskTagScope Scope(ETaskTag::EGameThread);
+
+	// start up the main loop
+	GEngineLoop.PreInit(ArgC, ArgV);
+
+	double StartTime = FPlatformTime::Seconds();
+
+	int32 Result = RunDiffAssetBulkData();
+
+	UE_LOG(LogDiffAssetBulk, Display, TEXT("Logging.."));
+
+	GLog->Flush();
+
+	RequestEngineExit(TEXT("DiffAssetBulkData Exiting"));
+
+	FEngineLoop::AppPreExit();
+	FModuleManager::Get().UnloadModulesAtShutdown();
+	FEngineLoop::AppExit();
+
+	return Result;
+}
+
