Adding moved package information, detailed new/delete listings, and cleaning up output a bit.
#rnx

diff --git a/Engine/Source/Programs/DiffAssetBulkData/Private/DiffAssetBulkData.cpp b/Engine/Source/Programs/DiffAssetBulkData/Private/DiffAssetBulkData.cpp
--- a/Engine/Source/Programs/DiffAssetBulkData/Private/DiffAssetBulkData.cpp
+++ b/Engine/Source/Programs/DiffAssetBulkData/Private/DiffAssetBulkData.cpp
@@ -84,6 +84,7 @@ static int32 RunDiffAssetBulkData()
 		UE_LOG(LogDiffAssetBulk, Display, TEXT(""));
 		UE_LOG(LogDiffAssetBulk, Display, TEXT("    -Base=<path/to/file>              Base Development Asset Registry (Required)"));
 		UE_LOG(LogDiffAssetBulk, Display, TEXT("    -Current=<path/to/file>           New Development Asset Registry (Required)"));
+		UE_LOG(LogDiffAssetBulk, Display, TEXT("    -Optional                         Evaluate Optional bulk data changes instead."));
 		UE_LOG(LogDiffAssetBulk, Display, TEXT("    -ListMixed                        Show the list of changed packages with assets that have matching"));
 		UE_LOG(LogDiffAssetBulk, Display, TEXT("                                      blame tags, but also assets without."));
 		UE_LOG(LogDiffAssetBulk, Display, TEXT("    -ListDeterminism                  Show the list of changed packages with assets that have matching"));
@@ -96,6 +97,7 @@ static int32 RunDiffAssetBulkData()
 		return 1;
 	}
 
+	bool bEvaluateOptional = FParse::Param(CmdLine, TEXT("Optional"));
 	bool bListMixed = FParse::Param(CmdLine, TEXT("ListMixed"));
 	bool bListDeterminism = FParse::Param(CmdLine, TEXT("ListDeterminism"));
 	bool bListUnrepresented = FParse::Param(CmdLine, TEXT("ListUnrepresented"));
@@ -105,16 +107,46 @@ static int32 RunDiffAssetBulkData()
 	FParse::Value(CmdLine, TEXT("ListNoBlame="), ListNoBlame);
 
 	FString ListCSV;
-	TUniquePtr<FArchive> OutputCSVAr;
+	TUniquePtr<FArchive> ChangedCSVAr;
+	TUniquePtr<FArchive> NewCSVAr;
+	TUniquePtr<FArchive> MovedCSVAr;
+	TUniquePtr<FArchive> DeletedCSVAr;
 	if (FParse::Value(CmdLine, TEXT("ListCSV="), ListCSV))
 	{
-		OutputCSVAr.Reset(IFileManager::Get().CreateFileWriter(*ListCSV, 0));
-		if (!OutputCSVAr)
+		FString Extension = FPaths::GetExtension(ListCSV);
+		FString Base = FPaths::ChangeExtension(ListCSV, TEXT(""));
+
+		ChangedCSVAr.Reset(IFileManager::Get().CreateFileWriter(*(Base + TEXT("Changed.") + Extension), 0));
+		if (!ChangedCSVAr)
+		{
+			UE_LOG(LogDiffAssetBulk, Error, TEXT("Unable to open output CSV file: %s"), *(Base + TEXT("Changed.") + Extension));
+			return false;
+		}
+		ChangedCSVAr->Logf(TEXT("Blame, Class, PackageName, BlameBefore, BlameAfter"));
+
+		NewCSVAr.Reset(IFileManager::Get().CreateFileWriter(*(Base + TEXT("New.") + Extension), 0));
+		if (!NewCSVAr)
+		{
+			UE_LOG(LogDiffAssetBulk, Error, TEXT("Unable to open output CSV file: %s"), *(Base + TEXT("New.") + Extension));
+			return false;
+		}
+		NewCSVAr->Logf(TEXT("Class, PackageName"));
+
+		DeletedCSVAr.Reset(IFileManager::Get().CreateFileWriter(*(Base + TEXT("Deleted.") + Extension), 0));
+		if (!DeletedCSVAr)
+		{
+			UE_LOG(LogDiffAssetBulk, Error, TEXT("Unable to open output CSV file: %s"), *(Base + TEXT("Deleted.") + Extension));
+			return false;
+		}
+		DeletedCSVAr->Logf(TEXT("Class, PackageName"));
+
+		MovedCSVAr.Reset(IFileManager::Get().CreateFileWriter(*(Base + TEXT("Moved.") + Extension), 0));
+		if (!MovedCSVAr)
 		{
-			UE_LOG(LogDiffAssetBulk, Error, TEXT("Unable to open output CSV file: %s"), *ListCSV);
+			UE_LOG(LogDiffAssetBulk, Error, TEXT("Unable to open output CSV file: %s"), *(Base + TEXT("Moved.") + Extension));
 			return false;
 		}
-		OutputCSVAr->Logf(TEXT("Blame, Class, PackageName, BlameBefore, BlameAfter"));
+		MovedCSVAr->Logf(TEXT("Class, PackageName"));
 	}
 
 	// Convert the static init help text to a map
@@ -157,8 +189,6 @@ static int32 RunDiffAssetBulkData()
 	const TMap<FName, const FAssetPackageData*>& BasePackages = BaseState.GetAssetPackageDataMap();
 	const TMap<FName, const FAssetPackageData*>& CurrentPackages = CurrentState.GetAssetPackageDataMap();
 
-	TArray<FName> NewPackages, DeletedPackages;
-
 	struct FIteratedPackage
 	{
 		FName Name = NAME_None;
@@ -175,8 +205,6 @@ static int32 RunDiffAssetBulkData()
 
 	uint64 CurrentTotalSize = 0;
 	uint64 BaseTotalSize = 0;
-	uint64 DeletedSize = 0;
-	uint64 NewSize = 0;
 
 	{
 		for (const TPair<FName, const FAssetPackageData*>& NamePackageDataPair : BasePackages)
@@ -188,18 +216,9 @@ static int32 RunDiffAssetBulkData()
 			if (BaseMIAsset && BaseMIAsset->GetTagValue(UE::AssetRegistry::Stage_ChunkCompressedSizeFName, BaseCompressedSize))
 			{
 				BaseTotalSize += BaseCompressedSize;
-				if (Current == nullptr)
-				{
-					DeletedSize += BaseCompressedSize;
-				}
 			}
 
 			UnionedPackages.Emplace(FIteratedPackage(NamePackageDataPair.Key, NamePackageDataPair.Value, Current));
-
-			if (Current == nullptr)
-			{
-				DeletedPackages.Add(NamePackageDataPair.Key);
-			}
 		}
 
 		for (const TPair<FName, const FAssetPackageData*>& NamePackageDataPair : CurrentPackages)
@@ -211,15 +230,10 @@ static int32 RunDiffAssetBulkData()
 			if (CurrentMIAsset && CurrentMIAsset->GetTagValue(UE::AssetRegistry::Stage_ChunkCompressedSizeFName, CurrentCompressedSize))
 			{
 				CurrentTotalSize += CurrentCompressedSize;
-				if (Base == nullptr)
-				{
-					NewSize += CurrentCompressedSize;
-				}
 			}
 
 			if (Base == nullptr)
 			{
-				NewPackages.Add(NamePackageDataPair.Key);
 				UnionedPackages.Emplace(FIteratedPackage(NamePackageDataPair.Key, nullptr, NamePackageDataPair.Value));
 			}
 		}
@@ -230,19 +244,62 @@ static int32 RunDiffAssetBulkData()
 	// This whole thing assumes that the index parameter of CreateIoChunkId is always 0. This is likely not going
 	// to be true with FDerivedData, once that gets turned on, but should be easy to update when the time comes.
 	//
-	TSet<FName> ChangedChunksByType[(uint32)EIoChunkType::MAX], NewChunksByType[(uint32)EIoChunkType::MAX], DeletedChunksByType[(uint32)EIoChunkType::MAX];
+
+	// Save off what hashes got deleted so we can try to find packages that moved and report those separately.
+	TMap<FIoHash, TArray<FName, TInlineAllocator<1>>> DeletedChunkPackagesByHash;
+
+	TSet<FName> PackagesWithChangedChunks;
+	TSet<FName> PackagesWithDeletedChunks;
+	TMap<FName, TArray<FIoHash, TInlineAllocator<1>>> PackagesWithNewChunks;
+
+	auto ShouldProcessChunk = [bEvaluateOptional](const FIoChunkId& ChunkId)
+	{
+		if (ChunkId.GetChunkType() != EIoChunkType::BulkData &&
+			ChunkId.GetChunkType() != EIoChunkType::OptionalBulkData &&
+			ChunkId.GetChunkType() != EIoChunkType::MemoryMappedBulkData)
+		{
+			return false;
+		}
+
+		bool bIsOptional = ChunkId.GetChunkType() == EIoChunkType::OptionalBulkData;
+		if (bEvaluateOptional)
+		{
+			return bIsOptional;
+		}
+		return !bIsOptional;
+	};
+
+	uint64 TotalChangedSize = 0;
+	TMap < FName /* PackageName */, TPair<uint64, uint64>> PackageSizes;
 	for (const FIteratedPackage& IteratedPackage : UnionedPackages)
 	{
 		const FAssetPackageData* BasePackage = IteratedPackage.Base;
 		const FAssetPackageData* CurrentPackage = IteratedPackage.Current;
 
+
+		// Get the size change.
+		// IoStoreUtilities puts the size of the package on the most important asset
+		const FAssetData* BaseMIAsset = UE::AssetRegistry::GetMostImportantAsset(BaseState.GetAssetsByPackageName(IteratedPackage.Name), UE::AssetRegistry::EGetMostImportantAssetFlags::IgnoreSkipClasses);
+		const FAssetData* CurrentMIAsset = UE::AssetRegistry::GetMostImportantAsset(CurrentState.GetAssetsByPackageName(IteratedPackage.Name), UE::AssetRegistry::EGetMostImportantAssetFlags::IgnoreSkipClasses);
+
+		uint64 BaseCompressedSize = 0;
+		if (BaseMIAsset)
+		{
+			BaseMIAsset->GetTagValue(UE::AssetRegistry::Stage_ChunkCompressedSizeFName, BaseCompressedSize);
+		}
+		uint64 CurrentCompressedSize = 0;
+		if (CurrentMIAsset)
+		{
+			CurrentMIAsset->GetTagValue(UE::AssetRegistry::Stage_ChunkCompressedSizeFName, CurrentCompressedSize);
+		}
+
+		PackageSizes.Add(IteratedPackage.Name, { BaseCompressedSize, CurrentCompressedSize });
+
 		if (BasePackage)
 		{
 			for (const TPair<FIoChunkId, FIoHash>& ChunkHashPair : BasePackage->ChunkHashes)
 			{
-				if (ChunkHashPair.Key.GetChunkType() != EIoChunkType::BulkData &&
-					ChunkHashPair.Key.GetChunkType() != EIoChunkType::OptionalBulkData &&
-					ChunkHashPair.Key.GetChunkType() != EIoChunkType::MemoryMappedBulkData)
+				if (!ShouldProcessChunk(ChunkHashPair.Key))
 					continue;
 
 				const FIoHash* CurrentHash = nullptr;
@@ -253,17 +310,24 @@ static int32 RunDiffAssetBulkData()
 
 				if (CurrentHash == nullptr)
 				{
-					TSet<FName>& Deleted = DeletedChunksByType[(uint32)ChunkHashPair.Key.GetChunkType()];
-					check(Deleted.Contains(IteratedPackage.Name) == false); // Because only 0 chunk index
-					Deleted.Add(IteratedPackage.Name);
+					PackagesWithDeletedChunks.Add(IteratedPackage.Name);
+					DeletedChunkPackagesByHash.FindOrAdd(ChunkHashPair.Value).Add(IteratedPackage.Name);
 					continue;
 				}
 
 				if (*CurrentHash != ChunkHashPair.Value)
 				{
-					TSet<FName>& Changed = ChangedChunksByType[(uint32)ChunkHashPair.Key.GetChunkType()];
-					check(Changed.Contains(IteratedPackage.Name) == false); // Because only 0 chunk index
-					Changed.Add(IteratedPackage.Name);
+					PackagesWithChangedChunks.Add(IteratedPackage.Name);
+
+					// All we can really do here is assume the entire package gets resent, which is not likely
+					// in the general case, but it _is_ reasonably likely in the cases where a package's bulk data changes,
+					// which happens to be what we select on.
+					// The counter argument is that it's possible that the bulk data is Very Large (i.e. multiple compression blocks), and only
+					// one block out of the entire thing changed.
+					if (BaseMIAsset && CurrentMIAsset)
+					{
+						TotalChangedSize += CurrentCompressedSize;
+					}
 				}
 			}
 		}
@@ -272,29 +336,127 @@ static int32 RunDiffAssetBulkData()
 		{
 			for (const TPair<FIoChunkId, FIoHash>& ChunkHashPair : CurrentPackage->ChunkHashes)
 			{
-				if (ChunkHashPair.Key.GetChunkType() != EIoChunkType::BulkData &&
-					ChunkHashPair.Key.GetChunkType() != EIoChunkType::OptionalBulkData &&
-					ChunkHashPair.Key.GetChunkType() != EIoChunkType::MemoryMappedBulkData)
+				if (!ShouldProcessChunk(ChunkHashPair.Key))
+				{
 					continue;
+				}
 
 				if (!BasePackage ||
 					BasePackage->ChunkHashes.Contains(ChunkHashPair.Key) == false)
 				{
-					TSet<FName>& New = NewChunksByType[(uint32)ChunkHashPair.Key.GetChunkType()];
-					check(New.Contains(IteratedPackage.Name) == false); // Because only 0 chunk index
-					New.Add(IteratedPackage.Name);
+					PackagesWithNewChunks.FindOrAdd(IteratedPackage.Name).Add(ChunkHashPair.Value);
 				}
 			}
 		}
 	}
 
-	// Get a unique list of changed packages.
-	TSet<FName> ChangedPackages;
-	for (uint32 ChunkTypeIndex = 0; ChunkTypeIndex < (uint32)EIoChunkType::MAX; ChunkTypeIndex++)
+	TMap<FName, FName> MovedPackagesFromTo;
+	
+	// Look over the new packages - if any of them have exact matching entries in the deleted list,
+	// then we assume it's a moved chunk and remove it from the new/delete lists.
+	for (const TPair<FName, TArray<FIoHash, TInlineAllocator<1>>>& PackageHashesPair : PackagesWithNewChunks)
 	{
-		ChangedPackages.Append(ChangedChunksByType[ChunkTypeIndex]);
+		// Make sure all chunks we know about moved from the same place. We expect this to be only 1 for now, so warn on it.
+		FName MovedFrom = NAME_None;
+		for (const FIoHash& NewHash : PackageHashesPair.Value)
+		{
+			const TArray<FName, TInlineAllocator<1>>* PackagesThatHadThisChunk = DeletedChunkPackagesByHash.Find(NewHash);
+			if (PackagesThatHadThisChunk == nullptr ||
+				PackagesThatHadThisChunk->Num() == 0)
+			{
+				MovedFrom = NAME_None;
+				break;
+			}
+
+			// Due to duplication we could theoretically have the exact same bulk data in a bunch of
+			// different packages, so we consider it a move if it's in any of them. This could fail
+			// if there were multiple chunks where one came from one package and the other came from a different one,
+			// seems unlikely.
+			if (MovedFrom.IsNone())
+			{
+				// Grab the first one...
+				MovedFrom = (*PackagesThatHadThisChunk)[0];
+			}
+			else
+			{
+				bool bFound = false;
+				for (const FName& PackageThatHadThisChunk : (*PackagesThatHadThisChunk))
+				{
+					if (MovedFrom == PackageThatHadThisChunk)
+					{
+						bFound = true;
+						break;
+					}
+				}
+
+				if (!bFound)
+				{
+					MovedFrom = NAME_None;
+					break;
+				}
+			}
+		}
+
+		if (MovedFrom.IsNone())
+		{
+			continue; // Not moved - actual new package.
+		}
+
+		// We also only allow path moves - this is because it's not uncommon for folks to duplicate something like a mesh
+		// and change the material and this can confuse our hash matching. 
+		// However, if it's a _Generated_ package we actually want to know because it might be an issue with the stability
+		// of the generator.
+		TStringBuilder<64> MovedFromStr;
+		MovedFromStr << MovedFrom;
+		if (!FCString::Stristr(*MovedFromStr, TEXT("_GENERATED_")))
+		{
+			// it's not generated, so make sure the name matches.
+			TStringBuilder<64> MovedToStr;
+			MovedToStr << PackageHashesPair.Key;
+
+			const TCHAR* MovedFromShortName = FCString::Strrchr(*MovedFromStr, TEXT('/'));
+			if (MovedFromShortName)
+			{
+				MovedFromShortName++;
+			}
+			const TCHAR* MovedToShortName = FCString::Strrchr(*MovedToStr, TEXT('/'));
+			if (MovedToShortName)
+			{
+				MovedToShortName++;
+			}
+
+			// If we have short names and they are different, we assume it's not an actual move.
+			if (MovedFromShortName && MovedToShortName && FCString::Stricmp(MovedFromShortName, MovedToShortName))
+			{
+				continue;
+			}
+		}
+
+		if (MovedPackagesFromTo.Contains(MovedFrom))
+		{
+			UE_LOG(LogDiffAssetBulk, Display, TEXT("Package %s appears to have moved twice. Perhaps duplicated multiple times and original deleted? Or Material change?"), *MovedFromStr);
+			UE_LOG(LogDiffAssetBulk, Display, TEXT("    Existing: %s"), *WriteToString<64>(*MovedPackagesFromTo.Find(MovedFrom)));
+			UE_LOG(LogDiffAssetBulk, Display, TEXT("         New: %s"), *WriteToString<64>(PackageHashesPair.Key));
+			continue;
+		}
+		MovedPackagesFromTo.Add(MovedFrom, PackageHashesPair.Key);
 	}
 
+	// Done with this, empty it so it's obvious if we try to use it.
+	DeletedChunkPackagesByHash.Empty();
+
+	// Once we have the list of moved packages, remove them from the deleted/new lists
+	for (const TPair<FName, FName>& MovedPackageFromTo : MovedPackagesFromTo)
+	{
+		if (!PackagesWithNewChunks.Remove(MovedPackageFromTo.Value))
+		{
+			UE_LOG(LogDiffAssetBulk, Warning, TEXT("Unable to remove moved package %s from the new list"), *WriteToString<64>(MovedPackageFromTo.Value));
+		}
+		if (!PackagesWithDeletedChunks.Remove(MovedPackageFromTo.Key))
+		{
+			UE_LOG(LogDiffAssetBulk, Warning, TEXT("Unable to remove moved package %s from the deleted list"), *WriteToString<64>(MovedPackageFromTo.Key));
+		}
+	}
 
 	//
 	// We know what bulk datas *packages* changed. Try and see if any of the assets in the package have
@@ -317,9 +479,7 @@ static int32 RunDiffAssetBulkData()
 	TArray<FName /* PackageName */> PackagesWithUnassignableDiffsAndUntaggedAssets;
 	TMap<FTopLevelAssetPath /* AssetClass */, TArray<FName /* PackageName */>> PackagesWithUnassignableDiffsByAssumedClass;
 	
-	uint64 TotalChangedSize = 0;
-	TMap < FName /* PackageName */, TPair<uint64, uint64>> PackageSizes;
-	for (const FName& ChangedPackageName : ChangedPackages)
+	for (const FName& ChangedPackageName : PackagesWithChangedChunks)
 	{
 		TConstArrayView<FAssetData const*> BaseAssetDatas = BaseState.GetAssetsByPackageName(ChangedPackageName);
 		TConstArrayView<FAssetData const*> CurrentAssetDatas = CurrentState.GetAssetsByPackageName(ChangedPackageName);
@@ -337,28 +497,6 @@ static int32 RunDiffAssetBulkData()
 			const FAssetData* BaseAssetData;
 			const FAssetData* CurrentAssetData;
 		};
-
-		// Get the size change.
-		const FAssetData* BaseMIAsset = UE::AssetRegistry::GetMostImportantAsset(BaseAssetDatas, UE::AssetRegistry::EGetMostImportantAssetFlags::IgnoreSkipClasses);
-		const FAssetData* CurrentMIAsset = UE::AssetRegistry::GetMostImportantAsset(CurrentAssetDatas, UE::AssetRegistry::EGetMostImportantAssetFlags::IgnoreSkipClasses);
-		if (BaseMIAsset && CurrentMIAsset)
-		{
-			// IoStoreUtilities puts the size of the package on the most important asset
-			uint64 BaseCompressedSize = 0;
-			uint64 CurrentCompressedSize = 0;
-			if (BaseMIAsset->GetTagValue(UE::AssetRegistry::Stage_ChunkCompressedSizeFName, BaseCompressedSize) &&
-				CurrentMIAsset->GetTagValue(UE::AssetRegistry::Stage_ChunkCompressedSizeFName, CurrentCompressedSize))
-			{
-				PackageSizes.Add(ChangedPackageName, {BaseCompressedSize, CurrentCompressedSize});
-
-				// All we can really do here is assume the entire package gets resent, which is not likely
-				// in the general case, but it _is_ reasonably likely in the cases where a package's bulk data changes,
-				// which happens to be what we select on.
-				// The counter argument is that it's possible that the bulk data is Very Large (i.e. multiple compression blocks), and only
-				// one block out of the entire thing changed.
-				TotalChangedSize += CurrentCompressedSize;
-			}
-		}
 		
 		// We want to find all the tags that are in both base/current.
 		TMap<FName /* AssetName */, TArray<FDiffTag>> PackageDiffTags;
@@ -493,232 +631,295 @@ static int32 RunDiffAssetBulkData()
 		UE_LOG(LogDiffAssetBulk, Display, TEXT("No package sizes found - stage with asset registry writeback (project settings -> packaging) to get package size info"));
 	}
 
-	int32 PackagesWithNoSize = ChangedPackages.Num() - PackageSizes.Num();
-	
-	int32 TotalNewChunks = 0, TotalChangedChunks = 0, TotalDeletedChunks = 0;
-	UE_LOG(LogDiffAssetBulk, Display, TEXT("Modifications By IoStore Chunk (only bulk data tracked atm):"));
-	UE_LOG(LogDiffAssetBulk, Display, TEXT(""));
-	UE_LOG(LogDiffAssetBulk, Display, TEXT("    ChunkType                   New    Deleted    Changed"));
-	for (uint32 ChunkTypeIndex = 0; ChunkTypeIndex < (uint8)EIoChunkType::MAX; ChunkTypeIndex++)
+	auto ProcessPackageClassAndSize = [](FAssetRegistryState& State, const FName& PackageName, uint64& SizeToUpdate, TMap<FTopLevelAssetPath, TArray<FName>>& PackagesByClassToUpdate)
 	{
-		EIoChunkType ChunkType = (EIoChunkType)ChunkTypeIndex;
-		if (ChunkType != EIoChunkType::BulkData &&
-			ChunkType != EIoChunkType::OptionalBulkData &&
-			ChunkType != EIoChunkType::MemoryMappedBulkData)
+		const FAssetData* MIAsset = UE::AssetRegistry::GetMostImportantAsset(State.GetAssetsByPackageName(PackageName), UE::AssetRegistry::EGetMostImportantAssetFlags::IgnoreSkipClasses);
+		if (MIAsset)
 		{
-			continue;
+			// IoStoreUtilities puts the size of the package on the most important asset
+			uint64 CurrentCompressedSize = 0;
+			if (MIAsset->GetTagValue(UE::AssetRegistry::Stage_ChunkCompressedSizeFName, CurrentCompressedSize))
+			{
+				SizeToUpdate += CurrentCompressedSize;
+			}
+
+			PackagesByClassToUpdate.FindOrAdd(MIAsset->AssetClassPath).Add(PackageName);
 		}
+	};
 
-		const TSet<FName>& NewChunksForType = NewChunksByType[ChunkTypeIndex];
-		const TSet<FName>& DeletedChunksForType = DeletedChunksByType[ChunkTypeIndex];
-		const TSet<FName>& ChangedChunksForType = ChangedChunksByType[ChunkTypeIndex];
+	auto SumPackageSizes = [&PackageSizes](const TArray<FName>& PackageList, bool bUseBaseSize)
+	{
+		uint64 Total = 0;
+		for (const FName& PackageName : PackageList)
+		{
+			TPair<uint64, uint64>* Sizes = PackageSizes.Find(PackageName);
+			if (Sizes)
+			{
+				Total += bUseBaseSize ? Sizes->Key : Sizes->Value;
+			}
+		}
+		return Total;
+	};
 
-		TotalNewChunks += NewChunksForType.Num();
-		TotalChangedChunks += ChangedChunksForType.Num();
-		TotalDeletedChunks += DeletedChunksForType.Num();
-		UE_LOG(LogDiffAssetBulk, Display, TEXT("    %-20s %10d %10d %10d"), *LexToString(ChunkType), NewChunksForType.Num(), DeletedChunksForType.Num(), ChangedChunksForType.Num());
+	TMap<FTopLevelAssetPath, TArray<FName>> NewPackagesByClass;
+	uint64 TotalNewPackagesSize = 0;
+	for (const TPair<FName, TArray<FIoHash, TInlineAllocator<1>>>& PackageHashesPair : PackagesWithNewChunks)
+	{
+		ProcessPackageClassAndSize(CurrentState, PackageHashesPair.Key, TotalNewPackagesSize, NewPackagesByClass);
 	}
-	UE_LOG(LogDiffAssetBulk, Display, TEXT("    ====================================================="));
 
-	UE_LOG(LogDiffAssetBulk, Display, TEXT("    %-20s %10d %10d %10d"), TEXT("Total"), TotalNewChunks, TotalDeletedChunks, TotalChangedChunks);
-
-	UE_LOG(LogDiffAssetBulk, Display, TEXT(""));
+	TMap<FTopLevelAssetPath, TArray<FName>> DeletedPackagesByClass;
+	uint64 TotalDeletedPackagesSize = 0;
+	for (const FName& DeletedPackage : PackagesWithDeletedChunks)
+	{
+		ProcessPackageClassAndSize(BaseState, DeletedPackage, TotalDeletedPackagesSize, DeletedPackagesByClass);
+	}
 
-	UE_LOG(LogDiffAssetBulk, Display, TEXT("    Base Packages:                %8d (%s bytes)"), BasePackages.Num(), *FText::AsNumber(BaseTotalSize).ToString());
-	UE_LOG(LogDiffAssetBulk, Display, TEXT("    Current Packages:             %8d (%s bytes)"), CurrentPackages.Num(), *FText::AsNumber(CurrentTotalSize).ToString());
-	UE_LOG(LogDiffAssetBulk, Display, TEXT("    Bulk Data Packages Added:     %8d (%s bytes)"), NewPackages.Num(), *FText::AsNumber(NewSize).ToString());
-	UE_LOG(LogDiffAssetBulk, Display, TEXT("    Bulk Data Packages Deleted:   %8d (%s bytes)"), DeletedPackages.Num(), *FText::AsNumber(DeletedSize).ToString());
-	UE_LOG(LogDiffAssetBulk, Display, TEXT("    Bulk Data Packages Changed:   %8d (%s bytes -- all chunks!)"), ChangedPackages.Num(), *FText::AsNumber(TotalChangedSize).ToString());
-	UE_LOG(LogDiffAssetBulk, Display, TEXT("    Packages with no size info:   %8d"), PackagesWithNoSize);
+	TMap<FTopLevelAssetPath, TArray<FName>> MovedPackagesByClass;
+	uint64 TotalMovedPackagesSize = 0;
+	for (const TPair<FName, FName>& MovedPackageFromTo : MovedPackagesFromTo)
+	{
+		ProcessPackageClassAndSize(BaseState, MovedPackageFromTo.Key, TotalMovedPackagesSize, MovedPackagesByClass);
+	}
 
-	if (ChangedPackages.Num() == 0)
+	int32 PackagesWithNoSize = UnionedPackages.Num() - PackageSizes.Num();
+	
+	UE_LOG(LogDiffAssetBulk, Display, TEXT("    ====================================================="));
+	if (bEvaluateOptional)
 	{
-		return 0;
+		UE_LOG(LogDiffAssetBulk, Display, TEXT("    OPTIONAL bulk data only"));
+	}
+	else
+	{
+		UE_LOG(LogDiffAssetBulk, Display, TEXT("    Excluding OPTIONAL bulk data chunks"));
 	}
 
 	UE_LOG(LogDiffAssetBulk, Display, TEXT(""));
 
-	TArray<FName>& CantDetermineAssetClassPackages = NoTagPackagesByAssumedClass.FindOrAdd(FTopLevelAssetPath());
+	UE_LOG(LogDiffAssetBulk, Display, TEXT("    Base Packages:                %8d %17s bytes"), BasePackages.Num(), *FText::AsNumber(BaseTotalSize).ToString());
+	UE_LOG(LogDiffAssetBulk, Display, TEXT("    Current Packages:             %8d %17s bytes"), CurrentPackages.Num(), *FText::AsNumber(CurrentTotalSize).ToString());
+	UE_LOG(LogDiffAssetBulk, Display, TEXT("    Bulk Data Packages Added:     %8d %17s bytes"), PackagesWithNewChunks.Num(), *FText::AsNumber(TotalNewPackagesSize).ToString());
+	UE_LOG(LogDiffAssetBulk, Display, TEXT("    Bulk Data Packages Deleted:   %8d %17s bytes"), PackagesWithDeletedChunks.Num(), *FText::AsNumber(TotalDeletedPackagesSize).ToString());
+	UE_LOG(LogDiffAssetBulk, Display, TEXT("    Bulk Data Packages Moved:     %8d %17s bytes"), MovedPackagesFromTo.Num(), *FText::AsNumber(TotalMovedPackagesSize).ToString());
+	UE_LOG(LogDiffAssetBulk, Display, TEXT("    Bulk Data Packages Changed:   %8d %17s bytes (all chunks!)"), PackagesWithChangedChunks.Num(), *FText::AsNumber(TotalChangedSize).ToString());
+	UE_LOG(LogDiffAssetBulk, Display, TEXT("    Packages with no size info:   %8d"), PackagesWithNoSize);
+	UE_LOG(LogDiffAssetBulk, Display, TEXT(""));
 
-	// Note this output is parsed by build scripts, be sure to fix those up if you change anything here.
-	UE_LOG(LogDiffAssetBulk, Display, TEXT("Changed package breakdown:"));
-	UE_LOG(LogDiffAssetBulk, Display, TEXT("    No blame information available:"));
+	if (PackagesWithChangedChunks.Num())
 	{
-		Algo::Sort(CantDetermineAssetClassPackages, FNameLexicalLess());
-		UE_LOG(LogDiffAssetBulk, Display, TEXT("        Can't determine asset class   : %-7d // Couldn't pick a representative asset in the package. -ListUnrepresented"), CantDetermineAssetClassPackages.Num());
-		if (bListUnrepresented)
+		TArray<FName>& CantDetermineAssetClassPackages = NoTagPackagesByAssumedClass.FindOrAdd(FTopLevelAssetPath());
+
+		// Note this output is parsed by build scripts, be sure to fix those up if you change anything here.
+		UE_LOG(LogDiffAssetBulk, Display, TEXT("Changed package breakdown:                               // -ListNoBlame=<class name>"));
+		UE_LOG(LogDiffAssetBulk, Display, TEXT("    No blame information available:"));
 		{
-			for (const FName& PackageName : CantDetermineAssetClassPackages)
+			Algo::Sort(CantDetermineAssetClassPackages, FNameLexicalLess());
+			UE_LOG(LogDiffAssetBulk, Display, TEXT("        Can't determine asset class           %6d     // Couldn't pick a representative asset in the package. -ListUnrepresented"), CantDetermineAssetClassPackages.Num());
+			if (bListUnrepresented)
 			{
-				UE_LOG(LogDiffAssetBulk, Display, TEXT("            %s"), *PackageName.ToString());
+				for (const FName& PackageName : CantDetermineAssetClassPackages)
+				{
+					UE_LOG(LogDiffAssetBulk, Display, TEXT("            %s"), *PackageName.ToString());
+				}
 			}
-		}
-		if (OutputCSVAr.IsValid())
-		{
-			for (const FName& PackageName : CantDetermineAssetClassPackages)
+			if (ChangedCSVAr.IsValid())
 			{
-				OutputCSVAr->Logf(TEXT("NoBlameInfo, Unknown, %s,,"), *WriteToString<64>(PackageName));
+				for (const FName& PackageName : CantDetermineAssetClassPackages)
+				{
+					ChangedCSVAr->Logf(TEXT("NoBlameInfo, Unknown, %s,,"), *WriteToString<64>(PackageName));
+				}
 			}
 		}
-	}
-	for (TPair<FTopLevelAssetPath, TArray<FName>>& ClassPackages : NoTagPackagesByAssumedClass)
-	{
-		if (ClassPackages.Key == FTopLevelAssetPath())
+		for (TPair<FTopLevelAssetPath, TArray<FName>>& ClassPackages : NoTagPackagesByAssumedClass)
 		{
-			continue;
-		}
-
-		uint64 TotalSizes = 0;
-		for (const FName& PackageName : ClassPackages.Value)
-		{
-			TPair<uint64, uint64>* Sizes = PackageSizes.Find(PackageName);
-			if (Sizes)
+			if (ClassPackages.Key == FTopLevelAssetPath()) // Skip packages we couldn't find a class for, handled above.
 			{
-				TotalSizes += Sizes->Value;
+				continue;
 			}
-		}
 
-		TStringBuilder<64> ClassName;
-		ClassName << ClassPackages.Key;
+			uint64 TotalSizes = SumPackageSizes(ClassPackages.Value, false);
 
-		UE_LOG(LogDiffAssetBulk, Display, TEXT("        %-30s: %d (%s bytes)  // -ListNoBlame=%s"), *ClassName, ClassPackages.Value.Num(), *FText::AsNumber(TotalSizes).ToString(), *ClassName);
-		if (ListNoBlame.Compare(TEXT("All"), ESearchCase::IgnoreCase) == 0 ||
-			ListNoBlame.Compare(ClassPackages.Key.ToString(), ESearchCase::IgnoreCase) == 0)
-		{
-			for (const FName& PackageName : ClassPackages.Value)
+			TStringBuilder<64> ClassName;
+			ClassName << ClassPackages.Key;
+
+			UE_LOG(LogDiffAssetBulk, Display, TEXT("        %-37s %6d %17s bytes"), *ClassName, ClassPackages.Value.Num(), *FText::AsNumber(TotalSizes).ToString());
+			if (ListNoBlame.Compare(TEXT("All"), ESearchCase::IgnoreCase) == 0 ||
+				ListNoBlame.Compare(ClassPackages.Key.ToString(), ESearchCase::IgnoreCase) == 0)
 			{
-				UE_LOG(LogDiffAssetBulk, Display, TEXT("        %s"), *WriteToString<64>(PackageName));
+				for (const FName& PackageName : ClassPackages.Value)
+				{
+					UE_LOG(LogDiffAssetBulk, Display, TEXT("        %s"), *WriteToString<64>(PackageName));
+				}
 			}
-		}
-		if (OutputCSVAr.IsValid())
-		{
-			for (const FName& PackageName : ClassPackages.Value)
+			if (ChangedCSVAr.IsValid())
 			{
-				OutputCSVAr->Logf(TEXT("NoBlameInfo, %s, %s,,"), *ClassName, *WriteToString<64>(PackageName));
+				for (const FName& PackageName : ClassPackages.Value)
+				{
+					ChangedCSVAr->Logf(TEXT("NoBlameInfo, %s, %s,,"), *ClassName, *WriteToString<64>(PackageName));
+				}
 			}
 		}
-	}
 
 
 
-	if (PackagesWithUnassignableDiffsByAssumedClass.Num())
-	{
-		int32 TotalUnassignablePackages = 0;
-		for (TPair<FTopLevelAssetPath, TArray<FName>>& ClassPackages : PackagesWithUnassignableDiffsByAssumedClass)
+		if (PackagesWithUnassignableDiffsByAssumedClass.Num())
 		{
-			TotalUnassignablePackages += ClassPackages.Value.Num();
-		}
+			int32 TotalUnassignablePackages = 0;
+			for (TPair<FTopLevelAssetPath, TArray<FName>>& ClassPackages : PackagesWithUnassignableDiffsByAssumedClass)
+			{
+				TotalUnassignablePackages += ClassPackages.Value.Num();
+			}
 		
 
-		UE_LOG(LogDiffAssetBulk, Display, TEXT("    Can't determine blame:        : %-7d // Assets had blame tags but all matched - check determinism! -ListDeterminism"), TotalUnassignablePackages);
-		for (TPair<FTopLevelAssetPath, TArray<FName>>& ClassPackages : PackagesWithUnassignableDiffsByAssumedClass)
-		{
-			uint64 TotalSizes = 0;
-			for (const FName& PackageName : ClassPackages.Value)
+			UE_LOG(LogDiffAssetBulk, Display, TEXT("    Can't determine blame:                    %6d     // Assets had blame tags but all matched - check determinism! -ListDeterminism"), TotalUnassignablePackages);
+			for (TPair<FTopLevelAssetPath, TArray<FName>>& ClassPackages : PackagesWithUnassignableDiffsByAssumedClass)
 			{
-				TPair<uint64, uint64>* Sizes = PackageSizes.Find(PackageName);
-				if (Sizes)
-				{
-					TotalSizes += Sizes->Value;
-				}
-			}
+				uint64 TotalSizes = SumPackageSizes(ClassPackages.Value, false);
 
-			TStringBuilder<64> ClassName;
-			ClassName << ClassPackages.Key;
+				TStringBuilder<64> ClassName;
+				ClassName << ClassPackages.Key;
 
-			UE_LOG(LogDiffAssetBulk, Display, TEXT("        %s : %d (%s bytes)"), *ClassName, ClassPackages.Value.Num(), *FText::AsNumber(TotalSizes).ToString());
-			Algo::Sort(ClassPackages.Value, FNameLexicalLess());
-			if (bListDeterminism)
-			{
-				for (const FName& PackageName : ClassPackages.Value)
+				UE_LOG(LogDiffAssetBulk, Display, TEXT("        %-37s %6d %17s bytes"), *ClassName, ClassPackages.Value.Num(), *FText::AsNumber(TotalSizes).ToString());
+				Algo::Sort(ClassPackages.Value, FNameLexicalLess());
+				if (bListDeterminism)
 				{
-					UE_LOG(LogDiffAssetBulk, Display, TEXT("            %s"), *WriteToString<64>(PackageName));
+					for (const FName& PackageName : ClassPackages.Value)
+					{
+						UE_LOG(LogDiffAssetBulk, Display, TEXT("            %s"), *WriteToString<64>(PackageName));
+					}
 				}
-			}
-			if (OutputCSVAr.IsValid())
-			{
-				for (const FName& PackageName : ClassPackages.Value)
+				if (ChangedCSVAr.IsValid())
 				{
-					OutputCSVAr->Logf(TEXT("NonDetermistic, %s, %s,,"), *ClassName, *WriteToString<64>(PackageName));
+					for (const FName& PackageName : ClassPackages.Value)
+					{
+						ChangedCSVAr->Logf(TEXT("NonDetermistic, %s, %s,,"), *ClassName, *WriteToString<64>(PackageName));
+					}
 				}
 			}
 		}
-	}
 
-	if (PackagesWithUnassignableDiffsAndUntaggedAssets.Num())
-	{
-		Algo::Sort(PackagesWithUnassignableDiffsAndUntaggedAssets, FNameLexicalLess());
-		
-		UE_LOG(LogDiffAssetBulk, Display, TEXT("    Potential untagged assets:    : %-7d // Package had assets with blame tags that matched, but also untagged assets. Might be determinism! -ListMixed"), PackagesWithUnassignableDiffsAndUntaggedAssets.Num());
-		if (bListMixed)
+		if (PackagesWithUnassignableDiffsAndUntaggedAssets.Num())
 		{
-			for (const FName& PackageName : PackagesWithUnassignableDiffsAndUntaggedAssets)
+			Algo::Sort(PackagesWithUnassignableDiffsAndUntaggedAssets, FNameLexicalLess());
+		
+			UE_LOG(LogDiffAssetBulk, Display, TEXT("    Potential untagged assets:          %6d     // Package had assets with blame tags that matched, but also untagged assets. Might be determinism! -ListMixed"), PackagesWithUnassignableDiffsAndUntaggedAssets.Num());
+			if (bListMixed)
 			{
-				UE_LOG(LogDiffAssetBulk, Display, TEXT("        %s"), *PackageName.ToString());
+				for (const FName& PackageName : PackagesWithUnassignableDiffsAndUntaggedAssets)
+				{
+					UE_LOG(LogDiffAssetBulk, Display, TEXT("        %s"), *PackageName.ToString());
+				}
 			}
-		}
-		if (OutputCSVAr.IsValid())
-		{
-			for (const FName& PackageName : PackagesWithUnassignableDiffsAndUntaggedAssets)
+			if (ChangedCSVAr.IsValid())
 			{
-				OutputCSVAr->Logf(TEXT("Mixed, Unknown, %s,,"), *WriteToString<64>(PackageName));
+				for (const FName& PackageName : PackagesWithUnassignableDiffsAndUntaggedAssets)
+				{
+					ChangedCSVAr->Logf(TEXT("Mixed, Unknown, %s,,"), *WriteToString<64>(PackageName));
+				}
 			}
 		}
-	}
-
-	if (Results.Num())
-	{
-		UE_LOG(LogDiffAssetBulk, Display, TEXT("    Summary changes by blame tag:"));
 
-		for (TPair<FName, TMap<FTopLevelAssetPath, TArray<FDiffResult>>>& TagResults : Results)
+		if (Results.Num())
 		{
-			uint32 TagCount = 0;
-			for (TPair<FTopLevelAssetPath, TArray<FDiffResult>>& ClassResults : TagResults.Value)
-			{
-				TagCount += ClassResults.Value.Num();
-			}
+			UE_LOG(LogDiffAssetBulk, Display, TEXT("    Summary changes by blame tag:                        // -ListBlame=<BlameTag>"));
 
-			const TCHAR** TagHelp = BuiltinDiffTagHelpMap.Find(TagResults.Key);
-			if (TagHelp != nullptr)
-			{
-				UE_LOG(LogDiffAssetBulk, Display, TEXT("        %-30s: %-7d // %s"), *TagResults.Key.ToString(), TagCount, *TagHelp);
-			}
-			else
+			for (TPair<FName, TMap<FTopLevelAssetPath, TArray<FDiffResult>>>& TagResults : Results)
 			{
-				UE_LOG(LogDiffAssetBulk, Display, TEXT("        %-30s: %-7d"), *TagResults.Key.ToString(), TagCount);
-			}
-		}
+				uint32 TagCount = 0;
+				for (TPair<FTopLevelAssetPath, TArray<FDiffResult>>& ClassResults : TagResults.Value)
+				{
+					TagCount += ClassResults.Value.Num();
+				}
 
-		UE_LOG(LogDiffAssetBulk, Display, TEXT("    Asset changes by blame tag:"));
+				TStringBuilder<32> TagName;
+				TagName << TagResults.Key;
 
-		for (TPair<FName, TMap<FTopLevelAssetPath, TArray<FDiffResult>>>& TagResults : Results)
-		{
-			UE_LOG(LogDiffAssetBulk, Display, TEXT("        %s  // -ListBlame=%s"), *TagResults.Key.ToString(), *TagResults.Key.ToString());
+				const TCHAR** TagHelp = BuiltinDiffTagHelpMap.Find(TagResults.Key);
+				if (TagHelp != nullptr)
+				{
+					UE_LOG(LogDiffAssetBulk, Display, TEXT("        %-37s %6d     // %s"), *TagName, TagCount, *TagHelp);
+				}
+				else
+				{
+					UE_LOG(LogDiffAssetBulk, Display, TEXT("        %-37s %6d"), *TagName, TagCount);
+				}
 
-			for (TPair<FTopLevelAssetPath, TArray<FDiffResult>>& ClassResults : TagResults.Value)
-			{
-				Algo::SortBy(ClassResults.Value, &FDiffResult::ChangedAssetObjectPath);
-				UE_LOG(LogDiffAssetBulk, Display, TEXT("            %s [%d]"), *ClassResults.Key.ToString(), ClassResults.Value.Num());
+				bool bListing = FCString::Stricmp(*ListBlame, TEXT("All")) == 0 ||
+					FCString::Stricmp(*ListBlame, *TagName) == 0;
 
-				if (ListBlame.Compare(TEXT("All"), ESearchCase::IgnoreCase) == 0 ||
-					ListBlame.Compare(ClassResults.Key.ToString(), ESearchCase::IgnoreCase) == 0)
+				for (TPair<FTopLevelAssetPath, TArray<FDiffResult>>& ClassResults : TagResults.Value)
 				{
-					for (FDiffResult& Result : ClassResults.Value)
+					Algo::SortBy(ClassResults.Value, &FDiffResult::ChangedAssetObjectPath);
+
+					if (bListing)
+					{
+						for (FDiffResult& Result : ClassResults.Value)
+						{
+							UE_LOG(LogDiffAssetBulk, Display, TEXT("                %s [%s -> %s]"), *Result.ChangedAssetObjectPath, *Result.TagBaseValue, *Result.TagCurrentValue);
+						}
+					}
+					if (ChangedCSVAr.IsValid())
 					{
-						UE_LOG(LogDiffAssetBulk, Display, TEXT("                %s [%s -> %s]"), *Result.ChangedAssetObjectPath, *Result.TagBaseValue, *Result.TagCurrentValue);
+						for (FDiffResult& Result : ClassResults.Value)
+						{
+							ChangedCSVAr->Logf(TEXT("%s, %s, %s, %s, %s"), *TagResults.Key.ToString(), *ClassResults.Key.ToString(), *WriteToString<64>(Result.ChangedAssetObjectPath), *Result.TagBaseValue, *Result.TagCurrentValue);
+						}
 					}
 				}
-				if (OutputCSVAr.IsValid())
+			}
+		}
+	} // end changed packages
+	
+	auto ProcessPackagesByClass = [&SumPackageSizes](const TMap<FTopLevelAssetPath, TArray<FName>>& PackagesByClass, FArchive* CSVArchive, const TMap<FName, FName>* PackageDestinationIfMoved, bool bUseBaseSizes)
+	{
+		for (const TPair<FTopLevelAssetPath, TArray<FName>>& PackagesForClass : PackagesByClass)
+		{
+			TStringBuilder<64> ClassName;
+			ClassName << PackagesForClass.Key;
+
+			uint64 TotalSize = SumPackageSizes(PackagesForClass.Value, bUseBaseSizes);
+
+			UE_LOG(LogDiffAssetBulk, Display, TEXT("    %-37s %6d %17s bytes"), *ClassName, PackagesForClass.Value.Num(), *FText::AsNumber(TotalSize).ToString(), *ClassName);
+
+			if (CSVArchive)
+			{
+				for (const FName& PackageName : PackagesForClass.Value)
 				{
-					for (FDiffResult& Result : ClassResults.Value)
+					if (PackageDestinationIfMoved)
 					{
-						OutputCSVAr->Logf(TEXT("%s, %s, %s, %s, %s"), *TagResults.Key.ToString(), *ClassResults.Key.ToString(), *WriteToString<64>(Result.ChangedAssetObjectPath), *Result.TagBaseValue, *Result.TagCurrentValue);
+						CSVArchive->Logf(TEXT("%s, %s, %s"), *ClassName, *WriteToString<64>(PackageName), *WriteToString<64>(*PackageDestinationIfMoved->Find(PackageName)));
+					}
+					else
+					{
+						CSVArchive->Logf(TEXT("%s, %s"), *ClassName, *WriteToString<64>(PackageName));
 					}
 				}
 			}
 		}
+	};
+
+	if (PackagesWithNewChunks.Num())
+	{
+		UE_LOG(LogDiffAssetBulk, Display, TEXT(""));
+		UE_LOG(LogDiffAssetBulk, Display, TEXT("New package breakdown:"));
+		ProcessPackagesByClass(NewPackagesByClass, NewCSVAr.Get(), nullptr, false);
 	}
 
+	if (PackagesWithDeletedChunks.Num())
+	{
+		UE_LOG(LogDiffAssetBulk, Display, TEXT(""));
+		UE_LOG(LogDiffAssetBulk, Display, TEXT("Deleted package breakdown:"));
+		ProcessPackagesByClass(DeletedPackagesByClass, DeletedCSVAr.Get(), nullptr, true);
+	}
+
+	if (MovedPackagesFromTo.Num())
+	{
+		UE_LOG(LogDiffAssetBulk, Display, TEXT(""));
+		UE_LOG(LogDiffAssetBulk, Display, TEXT("Moved package breakdown:"));
+		ProcessPackagesByClass(MovedPackagesByClass, MovedCSVAr.Get(), &MovedPackagesFromTo, true);
+	}
+
+
 	UE_LOG(LogDiffAssetBulk, Display, TEXT("Done."));
 
 	return 0;
