DiffAssetBulkData - way easier to parse output if there's no spaces here.
#rnx

diff --git a/Engine/Source/Programs/DiffAssetBulkData/Private/DiffAssetBulkData.cpp b/Engine/Source/Programs/DiffAssetBulkData/Private/DiffAssetBulkData.cpp
--- a/Engine/Source/Programs/DiffAssetBulkData/Private/DiffAssetBulkData.cpp
+++ b/Engine/Source/Programs/DiffAssetBulkData/Private/DiffAssetBulkData.cpp
@@ -714,7 +714,7 @@ static int32 RunDiffAssetBulkData()
 		UE_LOG(LogDiffAssetBulk, Display, TEXT("    No blame information available:"));
 		{
 			Algo::Sort(CantDetermineAssetClassPackages, FNameLexicalLess());
-			UE_LOG(LogDiffAssetBulk, Display, TEXT("        Can't determine asset class           %6d     // Couldn't pick a representative asset in the package. -ListUnrepresented"), CantDetermineAssetClassPackages.Num());
+			UE_LOG(LogDiffAssetBulk, Display, TEXT("        Unknown                               %6d     // Couldn't pick a representative asset in the package. -ListUnrepresented"), CantDetermineAssetClassPackages.Num());
 			if (bListUnrepresented)
 			{
 				for (const FName& PackageName : CantDetermineAssetClassPackages)
