Adding size information to CSV output

diff --git a/Engine/Source/Programs/DiffAssetBulkData/Private/DiffAssetBulkData.cpp b/Engine/Source/Programs/DiffAssetBulkData/Private/DiffAssetBulkData.cpp
--- a/Engine/Source/Programs/DiffAssetBulkData/Private/DiffAssetBulkData.cpp
+++ b/Engine/Source/Programs/DiffAssetBulkData/Private/DiffAssetBulkData.cpp
@@ -122,7 +122,7 @@ static int32 RunDiffAssetBulkData()
 			UE_LOG(LogDiffAssetBulk, Error, TEXT("Unable to open output CSV file: %s"), *(Base + TEXT("Changed.") + Extension));
 			return false;
 		}
-		ChangedCSVAr->Logf(TEXT("Blame, Class, PackageName, BlameBefore, BlameAfter"));
+		ChangedCSVAr->Logf(TEXT("Blame, Class, PackageName, BlameBefore, BlameAfter, OldCompressedSize, NewCompressedSize, OldUncompressedSize, NewUncompressedSize"));
 
 		NewCSVAr.Reset(IFileManager::Get().CreateFileWriter(*(Base + TEXT("New.") + Extension), 0));
 		if (!NewCSVAr)
@@ -269,8 +269,15 @@ static int32 RunDiffAssetBulkData()
 		return !bIsOptional;
 	};
 
+	struct FPackageSizes
+	{
+		uint64 BaseCompressedSize = 0;
+		uint64 CurrentCompressedSize = 0;
+		uint64 BaseUncompressedSize = 0;
+		uint64 CurrentUncompressedSize = 0;
+	};
 	uint64 TotalChangedSize = 0;
-	TMap < FName /* PackageName */, TPair<uint64, uint64>> PackageSizes;
+	TMap < FName /* PackageName */, FPackageSizes> PackageSizes;
 	for (const FIteratedPackage& IteratedPackage : UnionedPackages)
 	{
 		const FAssetPackageData* BasePackage = IteratedPackage.Base;
@@ -282,19 +289,19 @@ static int32 RunDiffAssetBulkData()
 		const FAssetData* BaseMIAsset = UE::AssetRegistry::GetMostImportantAsset(BaseState.GetAssetsByPackageName(IteratedPackage.Name), UE::AssetRegistry::EGetMostImportantAssetFlags::IgnoreSkipClasses);
 		const FAssetData* CurrentMIAsset = UE::AssetRegistry::GetMostImportantAsset(CurrentState.GetAssetsByPackageName(IteratedPackage.Name), UE::AssetRegistry::EGetMostImportantAssetFlags::IgnoreSkipClasses);
 
-		uint64 BaseCompressedSize = 0;
+		FPackageSizes& Sizes = PackageSizes.Add(IteratedPackage.Name);
+
 		if (BaseMIAsset)
 		{
-			BaseMIAsset->GetTagValue(UE::AssetRegistry::Stage_ChunkCompressedSizeFName, BaseCompressedSize);
+			BaseMIAsset->GetTagValue(UE::AssetRegistry::Stage_ChunkCompressedSizeFName, Sizes.BaseCompressedSize);
+			BaseMIAsset->GetTagValue(UE::AssetRegistry::Stage_ChunkSizeFName, Sizes.BaseUncompressedSize);
 		}
-		uint64 CurrentCompressedSize = 0;
 		if (CurrentMIAsset)
 		{
-			CurrentMIAsset->GetTagValue(UE::AssetRegistry::Stage_ChunkCompressedSizeFName, CurrentCompressedSize);
+			CurrentMIAsset->GetTagValue(UE::AssetRegistry::Stage_ChunkCompressedSizeFName, Sizes.CurrentCompressedSize);
+			CurrentMIAsset->GetTagValue(UE::AssetRegistry::Stage_ChunkSizeFName, Sizes.CurrentUncompressedSize);
 		}
 
-		PackageSizes.Add(IteratedPackage.Name, { BaseCompressedSize, CurrentCompressedSize });
-
 		if (BasePackage)
 		{
 			for (const TPair<FIoChunkId, FIoHash>& ChunkHashPair : BasePackage->ChunkHashes)
@@ -326,7 +333,7 @@ static int32 RunDiffAssetBulkData()
 					// one block out of the entire thing changed.
 					if (BaseMIAsset && CurrentMIAsset)
 					{
-						TotalChangedSize += CurrentCompressedSize;
+						TotalChangedSize += Sizes.CurrentCompressedSize;
 					}
 				}
 			}
@@ -626,11 +633,6 @@ static int32 RunDiffAssetBulkData()
 		}
 	}
 
-	if (PackageSizes.Num() == 0)
-	{
-		UE_LOG(LogDiffAssetBulk, Display, TEXT("No package sizes found - stage with asset registry writeback (project settings -> packaging) to get package size info"));
-	}
-
 	auto ProcessPackageClassAndSize = [](FAssetRegistryState& State, const FName& PackageName, uint64& SizeToUpdate, TMap<FTopLevelAssetPath, TArray<FName>>& PackagesByClassToUpdate)
 	{
 		const FAssetData* MIAsset = UE::AssetRegistry::GetMostImportantAsset(State.GetAssetsByPackageName(PackageName), UE::AssetRegistry::EGetMostImportantAssetFlags::IgnoreSkipClasses);
@@ -652,10 +654,10 @@ static int32 RunDiffAssetBulkData()
 		uint64 Total = 0;
 		for (const FName& PackageName : PackageList)
 		{
-			TPair<uint64, uint64>* Sizes = PackageSizes.Find(PackageName);
+			FPackageSizes* Sizes = PackageSizes.Find(PackageName);
 			if (Sizes)
 			{
-				Total += bUseBaseSize ? Sizes->Key : Sizes->Value;
+				Total += bUseBaseSize ? Sizes->BaseCompressedSize : Sizes->CurrentCompressedSize;
 			}
 		}
 		return Total;
@@ -755,7 +757,13 @@ static int32 RunDiffAssetBulkData()
 			{
 				for (const FName& PackageName : ClassPackages.Value)
 				{
-					ChangedCSVAr->Logf(TEXT("NoBlameInfo, %s, %s,,"), *ClassName, *WriteToString<64>(PackageName));
+					FPackageSizes* Sizes = PackageSizes.Find(PackageName);
+
+					ChangedCSVAr->Logf(TEXT("NoBlameInfo, %s, %s,,,%s,%s,%s,%s"), *ClassName, *WriteToString<64>(PackageName), 
+						Sizes ? *WriteToString<32>(Sizes->BaseCompressedSize) : TEXT(""), 
+						Sizes ? *WriteToString<32>(Sizes->CurrentCompressedSize) : TEXT(""),
+						Sizes ? *WriteToString<32>(Sizes->BaseUncompressedSize) : TEXT(""),
+						Sizes ? *WriteToString<32>(Sizes->CurrentUncompressedSize) : TEXT(""));
 				}
 			}
 		}
