Fix deprecation warnings, we could likely avoid the temporary TArrays but they should all be small

#rnx

diff --git a/Engine/Source/Programs/DiffAssetBulkData/Private/DiffAssetBulkData.cpp b/Engine/Source/Programs/DiffAssetBulkData/Private/DiffAssetBulkData.cpp
--- a/Engine/Source/Programs/DiffAssetBulkData/Private/DiffAssetBulkData.cpp
+++ b/Engine/Source/Programs/DiffAssetBulkData/Private/DiffAssetBulkData.cpp
@@ -211,7 +211,7 @@ static int32 RunDiffAssetBulkData()
 		{
 			const FAssetPackageData* Current = CurrentState.GetAssetPackageData(NamePackageDataPair.Key);
 
-			const FAssetData* BaseMIAsset = UE::AssetRegistry::GetMostImportantAsset(BaseState.GetAssetsByPackageName(NamePackageDataPair.Key), UE::AssetRegistry::EGetMostImportantAssetFlags::IgnoreSkipClasses);
+			const FAssetData* BaseMIAsset = UE::AssetRegistry::GetMostImportantAsset(BaseState.CopyAssetsByPackageName(NamePackageDataPair.Key), UE::AssetRegistry::EGetMostImportantAssetFlags::IgnoreSkipClasses);
 			uint64 BaseCompressedSize = 0;
 			if (BaseMIAsset && BaseMIAsset->GetTagValue(UE::AssetRegistry::Stage_ChunkCompressedSizeFName, BaseCompressedSize))
 			{
@@ -225,7 +225,7 @@ static int32 RunDiffAssetBulkData()
 		{
 			const FAssetPackageData* Base = BaseState.GetAssetPackageData(NamePackageDataPair.Key);
 
-			const FAssetData* CurrentMIAsset = UE::AssetRegistry::GetMostImportantAsset(CurrentState.GetAssetsByPackageName(NamePackageDataPair.Key), UE::AssetRegistry::EGetMostImportantAssetFlags::IgnoreSkipClasses);
+			const FAssetData* CurrentMIAsset = UE::AssetRegistry::GetMostImportantAsset(CurrentState.CopyAssetsByPackageName(NamePackageDataPair.Key), UE::AssetRegistry::EGetMostImportantAssetFlags::IgnoreSkipClasses);
 			uint64 CurrentCompressedSize = 0;
 			if (CurrentMIAsset && CurrentMIAsset->GetTagValue(UE::AssetRegistry::Stage_ChunkCompressedSizeFName, CurrentCompressedSize))
 			{
@@ -286,8 +286,8 @@ static int32 RunDiffAssetBulkData()
 
 		// Get the size change.
 		// IoStoreUtilities puts the size of the package on the most important asset
-		const FAssetData* BaseMIAsset = UE::AssetRegistry::GetMostImportantAsset(BaseState.GetAssetsByPackageName(IteratedPackage.Name), UE::AssetRegistry::EGetMostImportantAssetFlags::IgnoreSkipClasses);
-		const FAssetData* CurrentMIAsset = UE::AssetRegistry::GetMostImportantAsset(CurrentState.GetAssetsByPackageName(IteratedPackage.Name), UE::AssetRegistry::EGetMostImportantAssetFlags::IgnoreSkipClasses);
+		const FAssetData* BaseMIAsset = UE::AssetRegistry::GetMostImportantAsset(BaseState.CopyAssetsByPackageName(IteratedPackage.Name), UE::AssetRegistry::EGetMostImportantAssetFlags::IgnoreSkipClasses);
+		const FAssetData* CurrentMIAsset = UE::AssetRegistry::GetMostImportantAsset(CurrentState.CopyAssetsByPackageName(IteratedPackage.Name), UE::AssetRegistry::EGetMostImportantAssetFlags::IgnoreSkipClasses);
 
 		FPackageSizes& Sizes = PackageSizes.Add(IteratedPackage.Name);
 
@@ -488,8 +488,8 @@ static int32 RunDiffAssetBulkData()
 	
 	for (const FName& ChangedPackageName : PackagesWithChangedChunks)
 	{
-		TConstArrayView<FAssetData const*> BaseAssetDatas = BaseState.GetAssetsByPackageName(ChangedPackageName);
-		TConstArrayView<FAssetData const*> CurrentAssetDatas = CurrentState.GetAssetsByPackageName(ChangedPackageName);
+		TArray<FAssetData const*> BaseAssetDatas = BaseState.CopyAssetsByPackageName(ChangedPackageName);
+		TArray<FAssetData const*> CurrentAssetDatas = CurrentState.CopyAssetsByPackageName(ChangedPackageName);
 
 		struct FDiffTag
 		{
@@ -635,7 +635,7 @@ static int32 RunDiffAssetBulkData()
 
 	auto ProcessPackageClassAndSize = [](FAssetRegistryState& State, const FName& PackageName, uint64& SizeToUpdate, TMap<FTopLevelAssetPath, TArray<FName>>& PackagesByClassToUpdate)
 	{
-		const FAssetData* MIAsset = UE::AssetRegistry::GetMostImportantAsset(State.GetAssetsByPackageName(PackageName), UE::AssetRegistry::EGetMostImportantAssetFlags::IgnoreSkipClasses);
+		const FAssetData* MIAsset = UE::AssetRegistry::GetMostImportantAsset(State.CopyAssetsByPackageName(PackageName), UE::AssetRegistry::EGetMostImportantAssetFlags::IgnoreSkipClasses);
 		if (MIAsset)
 		{
 			// IoStoreUtilities puts the size of the package on the most important asset
