Fixed printfs.

diff --git a/Engine/Source/Programs/DiffAssetBulkData/Private/DiffAssetBulkData.cpp b/Engine/Source/Programs/DiffAssetBulkData/Private/DiffAssetBulkData.cpp
--- a/Engine/Source/Programs/DiffAssetBulkData/Private/DiffAssetBulkData.cpp
+++ b/Engine/Source/Programs/DiffAssetBulkData/Private/DiffAssetBulkData.cpp
@@ -887,7 +887,7 @@ static int32 RunDiffAssetBulkData()
 
 			uint64 TotalSize = SumPackageSizes(PackagesForClass.Value, bUseBaseSizes);
 
-			UE_LOG(LogDiffAssetBulk, Display, TEXT("    %-37s %6d %17s bytes"), *ClassName, PackagesForClass.Value.Num(), *FText::AsNumber(TotalSize).ToString(), *ClassName);
+			UE_LOG(LogDiffAssetBulk, Display, TEXT("    %-37s %6d %17s bytes"), *ClassName, PackagesForClass.Value.Num(), *FText::AsNumber(TotalSize).ToString());
 
 			if (CSVArchive)
 			{
